Fix crash in WebMMovieStreamer when there >1 startup movie (UE-70663).

- Also fixed VP8 and Vorbis not being accepted in startup movies.

diff --git a/Engine/Plugins/Media/WebMMedia/Source/WebMMedia/Private/Player/WebMVideoDecoder.cpp b/Engine/Plugins/Media/WebMMedia/Source/WebMMedia/Private/Player/WebMVideoDecoder.cpp
--- a/Engine/Plugins/Media/WebMMedia/Source/WebMMedia/Private/Player/WebMVideoDecoder.cpp
+++ b/Engine/Plugins/Media/WebMMedia/Source/WebMMedia/Private/Player/WebMVideoDecoder.cpp
@@ -178,7 +178,13 @@ void FWebMVideoDecoder::Close()
 	}
 
 	// Make sure all compute shader decoding is done
-	FlushRenderingCommands();
+	//
+	// This function can also be called on a rendering thread (the streamer is ticked there during a startup movie, and decoder gets deleted on StartNextMovie()
+	// if there are >1 movie queued). In this case we will ensure that the resources survive for one more frame after use by other means.
+	if (IsInGameThread())
+	{
+		FlushRenderingCommands();
+	}
 
 	if (bIsInitialized)
 	{
diff --git a/Engine/Plugins/Media/WebMMedia/Source/WebMMedia/Private/WebMContainer.cpp b/Engine/Plugins/Media/WebMMedia/Source/WebMMedia/Private/WebMContainer.cpp
--- a/Engine/Plugins/Media/WebMMedia/Source/WebMMedia/Private/WebMContainer.cpp
+++ b/Engine/Plugins/Media/WebMMedia/Source/WebMMedia/Private/WebMContainer.cpp
@@ -80,7 +80,7 @@ bool FWebMContainer::Open(const FString& FilePath)
 
 		if (Track->GetType() == mkvparser::Track::kVideo)
 		{
-			if (FCStringAnsi::Strcmp(Track->GetCodecId(), "V_VP9") == 0)
+			if (FCStringAnsi::Strcmp(Track->GetCodecId(), "V_VP9") == 0 || FCStringAnsi::Strcmp(Track->GetCodecId(), "V_VP8") == 0)
 			{
 				MkvFile->VideoTracks.Add(static_cast<const mkvparser::VideoTrack*>(Track));
 
@@ -99,7 +99,7 @@ bool FWebMContainer::Open(const FString& FilePath)
 		}
 		else if (Track->GetType() == mkvparser::Track::kAudio)
 		{
-			if (FCStringAnsi::Strcmp(Track->GetCodecId(), "A_OPUS") == 0)
+			if (FCStringAnsi::Strcmp(Track->GetCodecId(), "A_OPUS") == 0 || FCStringAnsi::Strcmp(Track->GetCodecId(), "A_VORBIS") == 0)
 			{
 				MkvFile->AudioTracks.Add(static_cast<const mkvparser::AudioTrack*>(Track));
 
diff --git a/Engine/Plugins/Runtime/WebMMoviePlayer/Source/WebMMoviePlayer/Private/WebMMovieStreamer.cpp b/Engine/Plugins/Runtime/WebMMoviePlayer/Source/WebMMoviePlayer/Private/WebMMovieStreamer.cpp
--- a/Engine/Plugins/Runtime/WebMMoviePlayer/Source/WebMMoviePlayer/Private/WebMMovieStreamer.cpp
+++ b/Engine/Plugins/Runtime/WebMMoviePlayer/Source/WebMMoviePlayer/Private/WebMMovieStreamer.cpp
@@ -27,6 +27,7 @@ FWebMMovieStreamer::FWebMMovieStreamer()
 	, VideoFramesCurrentlyProcessing(0)
 	, StartTime(0)
 	, bPlaying(false)
+	, TicksLeftToWaitPostCompletion(0)
 {
 }
 
@@ -134,6 +135,17 @@ bool FWebMMovieStreamer::Tick(float InDeltaTime)
 {
 	if (bPlaying)
 	{
+		if (TicksLeftToWaitPostCompletion)
+		{
+			if (--TicksLeftToWaitPostCompletion <= 0)
+			{
+				TicksLeftToWaitPostCompletion = 0;
+				return !StartNextMovie();
+			}
+
+			return false;
+		}
+
 		bool bHaveThingsToDo = false;
 
 		bHaveThingsToDo |= DisplayFrames(InDeltaTime);
@@ -141,16 +153,13 @@ bool FWebMMovieStreamer::Tick(float InDeltaTime)
 
 		bHaveThingsToDo |= ReadMoreFrames();
 
-		if (bHaveThingsToDo)
+		if (!bHaveThingsToDo)
 		{
-			// We're still playing this movie
-			return false;
-		}
-		else
-		{
-			// Try to start next movie from the queue
-			return !StartNextMovie();
+			// we're done playing this movie, make sure we can safely remove the textures next frame
+			TicksLeftToWaitPostCompletion = 1;
+			Viewport->SetTexture(nullptr);
 		}
+		return false;
 	}
 	else
 	{
diff --git a/Engine/Plugins/Runtime/WebMMoviePlayer/Source/WebMMoviePlayer/Private/WebMMovieStreamer.h b/Engine/Plugins/Runtime/WebMMoviePlayer/Source/WebMMoviePlayer/Private/WebMMovieStreamer.h
--- a/Engine/Plugins/Runtime/WebMMoviePlayer/Source/WebMMoviePlayer/Private/WebMMovieStreamer.h
+++ b/Engine/Plugins/Runtime/WebMMoviePlayer/Source/WebMMoviePlayer/Private/WebMMovieStreamer.h
@@ -49,6 +49,13 @@ private:
 	double StartTime;
 	bool bPlaying;
 
+	/** 
+	 * Number of ticks to wait after the movie is complete before moving on to the next one. 
+	 * 
+	 * This allows us to defer texture deletion while it is being displayed.
+	 */
+	int32 TicksLeftToWaitPostCompletion;
+
 	bool StartNextMovie();
 	void ReleaseAcquiredResources();
 	bool DisplayFrames(float InDeltaTime);
