WebM movie streamer: SDL audio support added

(merged from Dev-Console)

diff --git a/Engine/Plugins/Media/WebMMedia/Source/WebMMedia/Private/Player/WebMAudioDecoder.cpp b/Engine/Plugins/Media/WebMMedia/Source/WebMMedia/Private/Player/WebMAudioDecoder.cpp
--- a/Engine/Plugins/Media/WebMMedia/Source/WebMMedia/Private/Player/WebMAudioDecoder.cpp
+++ b/Engine/Plugins/Media/WebMMedia/Source/WebMMedia/Private/Player/WebMAudioDecoder.cpp
@@ -196,6 +196,11 @@ void FWebMAudioDecoder::DecodeAudioFramesAsync(const TArray<TSharedPtr<FWebMFram
 	}, TStatId(), nullptr, ENamedThreads::AnyThread);
 }
 
+bool FWebMAudioDecoder::IsBusy() const
+{
+	return AudioDecodingTask && !AudioDecodingTask->IsComplete();
+}
+
 void FWebMAudioDecoder::DoDecodeAudioFrames(const TArray<TSharedPtr<FWebMFrame>>& AudioFrames)
 {
 	for (const TSharedPtr<FWebMFrame>& AudioFrame : AudioFrames)
diff --git a/Engine/Plugins/Media/WebMMedia/Source/WebMMedia/Public/WebMAudioDecoder.h b/Engine/Plugins/Media/WebMMedia/Source/WebMMedia/Public/WebMAudioDecoder.h
--- a/Engine/Plugins/Media/WebMMedia/Source/WebMMedia/Public/WebMAudioDecoder.h
+++ b/Engine/Plugins/Media/WebMMedia/Source/WebMMedia/Public/WebMAudioDecoder.h
@@ -21,6 +21,7 @@ public:
 public:
 	bool Initialize(const char* InCodec, int32 InSampleRate, int32 InChannels, const uint8* CodecPrivateData, size_t CodecPrivateDataSize);
 	void DecodeAudioFramesAsync(const TArray<TSharedPtr<FWebMFrame>>& AudioFrames);
+	bool IsBusy() const;
 
 private:
 	enum ESupportedCodecs
diff --git a/Engine/Plugins/Media/WebMMedia/Source/WebMMedia/Public/WebMMediaAudioSample.h b/Engine/Plugins/Media/WebMMedia/Source/WebMMedia/Public/WebMMediaAudioSample.h
--- a/Engine/Plugins/Media/WebMMedia/Source/WebMMedia/Public/WebMMediaAudioSample.h
+++ b/Engine/Plugins/Media/WebMMedia/Source/WebMMedia/Public/WebMMediaAudioSample.h
@@ -99,6 +99,11 @@ public:
 		return Time;
 	}
 
+	const TArray<uint8>& GetDataBuffer() const
+	{
+		return Buffer;
+	}
+
 private:
 
 	/** The sample's data buffer. */
diff --git a/Engine/Plugins/Runtime/WebMMoviePlayer/Source/WebMMoviePlayer/Private/Audio/Null/WebMAudioBackend.h b/Engine/Plugins/Runtime/WebMMoviePlayer/Source/WebMMoviePlayer/Private/Audio/Null/WebMAudioBackend.h
new file mode 100644
--- /dev/null
+++ b/Engine/Plugins/Runtime/WebMMoviePlayer/Source/WebMMoviePlayer/Private/Audio/Null/WebMAudioBackend.h
@@ -0,0 +1,17 @@
+// Copyright 1998-2018 Epic Games, Inc. All Rights Reserved.
+
+#pragma once
+
+#include "CoreMinimal.h"
+
+class FWebMAudioBackendNull
+{
+public:
+	bool InitializePlatform() { return true;  }
+	void ShutdownPlatform() {}
+	bool StartStreaming(int32 SampleRate, int32 NumOfChannels) { return true; }
+	void StopStreaming() {}
+	bool SendAudio(const uint8* Buffer, size_t BufferSize) { return true; }
+};
+
+using FWebMAudioBackend = FWebMAudioBackendNull;
diff --git a/Engine/Plugins/Runtime/WebMMoviePlayer/Source/WebMMoviePlayer/Private/Audio/Unix/WebMAudioBackend.h b/Engine/Plugins/Runtime/WebMMoviePlayer/Source/WebMMoviePlayer/Private/Audio/Unix/WebMAudioBackend.h
new file mode 100644
--- /dev/null
+++ b/Engine/Plugins/Runtime/WebMMoviePlayer/Source/WebMMoviePlayer/Private/Audio/Unix/WebMAudioBackend.h
@@ -0,0 +1,30 @@
+// Copyright 1998-2018 Epic Games, Inc. All Rights Reserved.
+
+#pragma once
+
+#include "CoreMinimal.h"
+
+THIRD_PARTY_INCLUDES_START
+#include "SDL.h"
+#include "SDL_audio.h"
+THIRD_PARTY_INCLUDES_END
+
+class FWebMAudioBackendSDL
+{
+public:
+	FWebMAudioBackendSDL();
+	virtual ~FWebMAudioBackendSDL();
+
+	//~ IWebMAudioBackend interface
+	virtual bool InitializePlatform();
+	virtual void ShutdownPlatform();
+	virtual bool StartStreaming(int32 SampleRate, int32 NumOfChannels);
+	virtual void StopStreaming();
+	virtual bool SendAudio(const uint8* Buffer, size_t BufferSize);
+
+private:
+	SDL_AudioDeviceID AudioDevice;
+	bool bSDLInitialized;
+};
+
+using FWebMAudioBackend = FWebMAudioBackendSDL;
diff --git a/Engine/Plugins/Runtime/WebMMoviePlayer/Source/WebMMoviePlayer/Private/Audio/Unix/WebMAudioBackendSDL.cpp b/Engine/Plugins/Runtime/WebMMoviePlayer/Source/WebMMoviePlayer/Private/Audio/Unix/WebMAudioBackendSDL.cpp
new file mode 100644
--- /dev/null
+++ b/Engine/Plugins/Runtime/WebMMoviePlayer/Source/WebMMoviePlayer/Private/Audio/Unix/WebMAudioBackendSDL.cpp
@@ -0,0 +1,92 @@
+// Copyright 1998-2018 Epic Games, Inc. All Rights Reserved.
+
+#include "WebMAudioBackend.h"
+#include "WebMMovieCommon.h"
+
+FWebMAudioBackendSDL::FWebMAudioBackendSDL()
+	: AudioDevice(0)
+	, bSDLInitialized(false)
+{
+}
+
+FWebMAudioBackendSDL::~FWebMAudioBackendSDL()
+{
+	ShutdownPlatform();
+}
+
+bool FWebMAudioBackendSDL::InitializePlatform()
+{
+	int32 Result = SDL_InitSubSystem(SDL_INIT_AUDIO);
+	if (Result < 0)
+	{
+		UE_LOG(LogWebMMoviePlayer, Error, TEXT("SDL_InitSubSystem create failed: %d"), Result);
+		bSDLInitialized = false;
+		return false;
+	}
+	else
+	{
+		bSDLInitialized = true;
+		return true;
+	}
+}
+
+void FWebMAudioBackendSDL::ShutdownPlatform()
+{
+	StopStreaming();
+
+	if (bSDLInitialized)
+	{
+		SDL_Quit();
+	}
+}
+
+bool FWebMAudioBackendSDL::StartStreaming(int32 SampleRate, int32 NumOfChannels)
+{
+	if (!bSDLInitialized)
+	{
+		return false;
+	}
+
+	SDL_AudioSpec AudioSpec;
+	memset(&AudioSpec, 0, sizeof(SDL_AudioSpec));
+	AudioSpec.freq = SampleRate;
+	AudioSpec.format = AUDIO_S16;
+	AudioSpec.channels = NumOfChannels;
+	AudioSpec.samples = 4096;
+
+	AudioDevice = SDL_OpenAudioDevice(nullptr, 0, &AudioSpec, nullptr, 0);
+
+	if (!AudioDevice)
+	{
+		UE_LOG(LogWebMMoviePlayer, Error, TEXT("SDL_OpenAudioDevice failed"));
+		return false;
+	}
+	else
+	{
+		return true;
+	}
+}
+
+void FWebMAudioBackendSDL::StopStreaming()
+{
+	if (AudioDevice)
+	{
+		SDL_CloseAudioDevice(AudioDevice);
+		AudioDevice = 0;
+	}
+}
+
+bool FWebMAudioBackendSDL::SendAudio(const uint8* Buffer, size_t BufferSize)
+{
+	int32 Result = SDL_QueueAudio(AudioDevice, Buffer, BufferSize);
+	if (Result < 0)
+	{
+		UE_LOG(LogWebMMoviePlayer, Error, TEXT("SDL_QueueAudio failed: %d"), Result);
+		return false;
+	}
+	else
+	{
+		SDL_PauseAudioDevice(AudioDevice, 0);
+		return true;
+	}
+}
diff --git a/Engine/Plugins/Runtime/WebMMoviePlayer/Source/WebMMoviePlayer/Private/WebMMovieStreamer.cpp b/Engine/Plugins/Runtime/WebMMoviePlayer/Source/WebMMoviePlayer/Private/WebMMovieStreamer.cpp
--- a/Engine/Plugins/Runtime/WebMMoviePlayer/Source/WebMMoviePlayer/Private/WebMMovieStreamer.cpp
+++ b/Engine/Plugins/Runtime/WebMMoviePlayer/Source/WebMMoviePlayer/Private/WebMMovieStreamer.cpp
@@ -13,6 +13,7 @@
 #include "WebMAudioDecoder.h"
 #include "WebMMediaFrame.h"
 #include "WebMContainer.h"
+#include "WebMMediaAudioSample.h"
 #include "WebMMediaTextureSample.h"
 
 DEFINE_LOG_CATEGORY(LogWebMMoviePlayer);
@@ -24,18 +25,27 @@ FWebMMovieStreamer::FWebMMovieStreamer()
 	, bPlaying(false)
 {
 	Viewport = MakeShareable(new FMovieViewport());
+
+	AudioBackend = MakeUnique<FWebMAudioBackend>();
+	AudioBackend->InitializePlatform();
 }
 
 FWebMMovieStreamer::~FWebMMovieStreamer()
 {
 	Cleanup();
+
+	AudioBackend->ShutdownPlatform();
 }
 
 void FWebMMovieStreamer::Cleanup()
 {
-	ForceCompletion();
+	VideoDecoder.Reset();
+	AudioDecoder.Reset();
+	Container.Reset();
 
 	ReleaseAcquiredResources();
+
+	AudioBackend->StopStreaming();
 }
 
 FTexture2DRHIRef FWebMMovieStreamer::GetTexture()
@@ -56,13 +66,14 @@ bool FWebMMovieStreamer::Init(const TArray<FString>& InMoviePaths, TEnumAsByte<E
 
 	// Play the next movie in the queue
 	return bPlaying;
-
 }
 
 bool FWebMMovieStreamer::StartNextMovie()
 {
 	if (MovieQueue.Num() > 0)
 	{
+		Cleanup();
+
 		MovieName = MovieQueue[0];
 
 		MovieQueue.RemoveAt(0);
@@ -98,6 +109,8 @@ bool FWebMMovieStreamer::StartNextMovie()
 		AudioDecoder->Initialize(DefaultAudioTrack.CodecName, DefaultAudioTrack.SampleRate, DefaultAudioTrack.NumOfChannels, DefaultAudioTrack.CodecPrivateData, DefaultAudioTrack.CodecPrivateDataSize);
 		VideoDecoder->Initialize(DefaultVideoTrack.CodecName);
 
+		AudioBackend->StartStreaming(DefaultAudioTrack.SampleRate, DefaultAudioTrack.NumOfChannels);
+
 		CurrentTime = 0;
 
 		return true;
@@ -126,8 +139,7 @@ bool FWebMMovieStreamer::Tick(float InDeltaTime)
 		bool bHaveThingsToDo = false;
 
 		bHaveThingsToDo |= DisplayFrames(InDeltaTime);
-
-		// TODO! Audio!
+		bHaveThingsToDo |= SendAudio(InDeltaTime);
 
 		bHaveThingsToDo |= DecodeMoreFrames();
 
@@ -179,10 +191,6 @@ void FWebMMovieStreamer::ReleaseAcquiredResources()
 	}
 
 	Viewport->SetTexture(nullptr);
-
-	VideoDecoder.Reset();
-	AudioDecoder.Reset();
-	Container.Reset();
 }
 
 bool FWebMMovieStreamer::DisplayFrames(float InDeltaTime)
@@ -210,6 +218,22 @@ bool FWebMMovieStreamer::DisplayFrames(float InDeltaTime)
 	return Samples->NumVideoSamples() > 0 || VideoDecoder->IsBusy();
 }
 
+bool FWebMMovieStreamer::SendAudio(float InDeltaTime)
+{
+	TSharedPtr<IMediaAudioSample, ESPMode::ThreadSafe> AudioSample;
+	TRange<FTimespan> TimeRange(FTimespan::Zero(), FTimespan::FromSeconds(CurrentTime));
+	bool bFoundSample = Samples->FetchAudio(TimeRange, AudioSample);
+
+	if (bFoundSample && AudioSample)
+	{
+		FWebMMediaAudioSample* WebMSample = StaticCast<FWebMMediaAudioSample*>(AudioSample.Get());
+
+		AudioBackend->SendAudio(WebMSample->GetDataBuffer().GetData(), WebMSample->GetDataBuffer().Num());
+	}
+
+	return Samples->NumAudio() > 0 || AudioDecoder->IsBusy();
+}
+
 bool FWebMMovieStreamer::DecodeMoreFrames()
 {
 	// Read frames up to 1 secs in the future
diff --git a/Engine/Plugins/Runtime/WebMMoviePlayer/Source/WebMMoviePlayer/Private/WebMMovieStreamer.h b/Engine/Plugins/Runtime/WebMMoviePlayer/Source/WebMMoviePlayer/Private/WebMMovieStreamer.h
--- a/Engine/Plugins/Runtime/WebMMoviePlayer/Source/WebMMoviePlayer/Private/WebMMovieStreamer.h
+++ b/Engine/Plugins/Runtime/WebMMoviePlayer/Source/WebMMoviePlayer/Private/WebMMovieStreamer.h
@@ -3,6 +3,7 @@
 #pragma once
 
 #include "MoviePlayer.h"
+#include "WebMAudioBackend.h"
 
 class FWebMVideoDecoder;
 class FWebMAudioDecoder;
@@ -43,6 +44,7 @@ private:
 	TUniquePtr<FWebMVideoDecoder> VideoDecoder;
 	TUniquePtr<FWebMAudioDecoder> AudioDecoder;
 	TUniquePtr<FWebMContainer> Container;
+	TUniquePtr<FWebMAudioBackend> AudioBackend;
 	TSharedPtr<FMediaSamples, ESPMode::ThreadSafe> Samples;
 	TSharedPtr<FMovieViewport> Viewport;
 	float CurrentTime;
@@ -51,5 +53,6 @@ private:
 	bool StartNextMovie();
 	void ReleaseAcquiredResources();
 	bool DisplayFrames(float InDeltaTime);
+	bool SendAudio(float InDeltaTime);
 	bool DecodeMoreFrames();
 };
diff --git a/Engine/Plugins/Runtime/WebMMoviePlayer/Source/WebMMoviePlayer/WebMMoviePlayer.Build.cs b/Engine/Plugins/Runtime/WebMMoviePlayer/Source/WebMMoviePlayer/WebMMoviePlayer.Build.cs
--- a/Engine/Plugins/Runtime/WebMMoviePlayer/Source/WebMMoviePlayer/WebMMoviePlayer.Build.cs
+++ b/Engine/Plugins/Runtime/WebMMoviePlayer/Source/WebMMoviePlayer/WebMMoviePlayer.Build.cs
@@ -21,6 +21,17 @@ namespace UnrealBuildTool.Rules
 					"MediaUtils",
 					"WebMMedia",
 				});
+
+			if (Target.IsInPlatformGroup(UnrealPlatformGroup.Unix))
+			{
+				PrivateDependencyModuleNames.Add("SDL2");
+
+				PrivateIncludePaths.Add("WebMMoviePlayer/Private/Audio/Unix");
+			}
+			else
+			{
+				PrivateIncludePaths.Add("WebMMoviePlayer/Private/Audio/Null");
+			}
 		}
 	}
 }
