WebM Movie Streamer: data read in backround thread

diff --git a/Engine/Plugins/Runtime/WebMMoviePlayer/Source/WebMMoviePlayer/Private/WebMMovieStreamer.cpp b/Engine/Plugins/Runtime/WebMMoviePlayer/Source/WebMMoviePlayer/Private/WebMMovieStreamer.cpp
--- a/Engine/Plugins/Runtime/WebMMoviePlayer/Source/WebMMoviePlayer/Private/WebMMovieStreamer.cpp
+++ b/Engine/Plugins/Runtime/WebMMoviePlayer/Source/WebMMoviePlayer/Private/WebMMovieStreamer.cpp
@@ -15,18 +15,64 @@
 #include "WebMContainer.h"
 #include "WebMMediaAudioSample.h"
 #include "WebMMediaTextureSample.h"
+#include "GenericPlatform/GenericPlatformProcess.h"
+#include "HAL/Runnable.h"
+#include "HAL/RunnableThread.h"
+#include "Misc/SingleThreadRunnable.h"
 
 DEFINE_LOG_CATEGORY(LogWebMMoviePlayer);
 
+class FWebMMovieStreamer::FWebMBackgroundReader : public FRunnable, private FSingleThreadRunnable
+{
+public:
+	FWebMBackgroundReader(FWebMMovieStreamer& InStreamer)
+		: Streamer(InStreamer)
+		, bIsFinished(false)
+	{
+	}
+
+	//~ FRunnable interface
+	virtual FSingleThreadRunnable* GetSingleThreadInterface() override
+	{
+		return this;
+	}
+
+	virtual uint32 Run() override
+	{
+		bIsFinished = false;
+
+		while (!bIsFinished && Streamer.IsPlaying())
+		{
+			bIsFinished = Streamer.DecodeMoreFramesInAnotherThread();
+		}
+
+		return 0;
+	}
+
+protected:
+	//~ FSingleThreadRunnable interface
+	virtual void Tick() override
+	{
+		if (!bIsFinished && Streamer.IsPlaying())
+		{
+			bIsFinished = Streamer.DecodeMoreFramesInAnotherThread();
+		}
+	}
+
+private:
+	FWebMMovieStreamer& Streamer;
+	bool bIsFinished;
+};
+
 FWebMMovieStreamer::FWebMMovieStreamer()
 	: CurrentTexture(0)
-	, Samples(new FMediaSamples())
+	, AudioBackend(MakeUnique<FWebMAudioBackend>())
+	, BackgroundReader(MakeUnique<FWebMBackgroundReader>(*this))
+	, Samples(MakeShareable(new FMediaSamples()))
+	, Viewport(MakeShareable(new FMovieViewport()))
 	, CurrentTime(0)
 	, bPlaying(false)
 {
-	Viewport = MakeShareable(new FMovieViewport());
-
-	AudioBackend = MakeUnique<FWebMAudioBackend>();
 	AudioBackend->InitializePlatform();
 }
 
@@ -39,6 +85,14 @@ FWebMMovieStreamer::~FWebMMovieStreamer()
 
 void FWebMMovieStreamer::Cleanup()
 {
+	bPlaying = false;
+
+	if (BackgroundReaderThread)
+	{
+		BackgroundReaderThread->WaitForCompletion();
+		BackgroundReaderThread.Reset();
+	}
+
 	VideoDecoder.Reset();
 	AudioDecoder.Reset();
 	Container.Reset();
@@ -62,10 +116,7 @@ bool FWebMMovieStreamer::Init(const TArray<FString>& InMoviePaths, TEnumAsByte<E
 	MovieQueue.Append(InMoviePaths);
 
 	// start our first movie playing
-	bPlaying = StartNextMovie();
-
-	// Play the next movie in the queue
-	return bPlaying;
+	return StartNextMovie();
 }
 
 bool FWebMMovieStreamer::StartNextMovie()
@@ -112,6 +163,9 @@ bool FWebMMovieStreamer::StartNextMovie()
 		AudioBackend->StartStreaming(DefaultAudioTrack.SampleRate, DefaultAudioTrack.NumOfChannels);
 
 		CurrentTime = 0;
+		bPlaying = true;
+
+		BackgroundReaderThread.Reset(FRunnableThread::Create(BackgroundReader.Get(), TEXT("FWebMBackgroundReader"), 0, TPri_BelowNormal));
 
 		return true;
 	}
@@ -141,11 +195,18 @@ bool FWebMMovieStreamer::Tick(float InDeltaTime)
 		bHaveThingsToDo |= DisplayFrames(InDeltaTime);
 		bHaveThingsToDo |= SendAudio(InDeltaTime);
 
-		bHaveThingsToDo |= DecodeMoreFrames();
-
 		CurrentTime += InDeltaTime;
 
-		return !bHaveThingsToDo;
+		if (bHaveThingsToDo)
+		{
+			// We're still playing this movie
+			return false;
+		}
+		else
+		{
+			// Try to start next movie from the queue
+			return !StartNextMovie();
+		}
 	}
 	else
 	{
@@ -220,23 +281,26 @@ bool FWebMMovieStreamer::DisplayFrames(float InDeltaTime)
 
 bool FWebMMovieStreamer::SendAudio(float InDeltaTime)
 {
+	// Just send all available audio for processing
+	
 	TSharedPtr<IMediaAudioSample, ESPMode::ThreadSafe> AudioSample;
-	TRange<FTimespan> TimeRange(FTimespan::Zero(), FTimespan::FromSeconds(CurrentTime));
+	TRange<FTimespan> TimeRange(FTimespan::Zero(), FTimespan::MaxValue());
 	bool bFoundSample = Samples->FetchAudio(TimeRange, AudioSample);
 
-	if (bFoundSample && AudioSample)
+	while (bFoundSample && AudioSample)
 	{
 		FWebMMediaAudioSample* WebMSample = StaticCast<FWebMMediaAudioSample*>(AudioSample.Get());
-
 		AudioBackend->SendAudio(WebMSample->GetDataBuffer().GetData(), WebMSample->GetDataBuffer().Num());
+
+		bFoundSample = Samples->FetchAudio(TimeRange, AudioSample);
 	}
 
+
 	return Samples->NumAudio() > 0 || AudioDecoder->IsBusy();
 }
 
-bool FWebMMovieStreamer::DecodeMoreFrames()
+bool FWebMMovieStreamer::DecodeMoreFramesInAnotherThread()
 {
-	// Read frames up to 1 secs in the future
 	FTimespan ReadBufferLength = FTimespan::FromSeconds(0.1);
 
 	TArray<TSharedPtr<FWebMFrame>> AudioFrames;
@@ -256,5 +320,5 @@ bool FWebMMovieStreamer::DecodeMoreFrames()
 		AudioDecoder->DecodeAudioFramesAsync(AudioFrames);
 	}
 
-	return VideoFrames.Num() > 0 || AudioFrames.Num() > 0;
+	return VideoFrames.Num() == 0 && AudioFrames.Num() == 0;
 }
diff --git a/Engine/Plugins/Runtime/WebMMoviePlayer/Source/WebMMoviePlayer/Private/WebMMovieStreamer.h b/Engine/Plugins/Runtime/WebMMoviePlayer/Source/WebMMoviePlayer/Private/WebMMovieStreamer.h
--- a/Engine/Plugins/Runtime/WebMMoviePlayer/Source/WebMMoviePlayer/Private/WebMMovieStreamer.h
+++ b/Engine/Plugins/Runtime/WebMMoviePlayer/Source/WebMMoviePlayer/Private/WebMMovieStreamer.h
@@ -9,6 +9,7 @@ class FWebMVideoDecoder;
 class FWebMAudioDecoder;
 class FMediaSamples;
 class FWebMContainer;
+class FRunnableThread;
 
 class FWebMMovieStreamer : public IMovieStreamer
 {
@@ -30,7 +31,12 @@ public:
 	FOnCurrentMovieClipFinished OnCurrentMovieClipFinishedDelegate;
 	virtual FOnCurrentMovieClipFinished& OnCurrentMovieClipFinished() override { return OnCurrentMovieClipFinishedDelegate; }
 
+	bool DecodeMoreFramesInAnotherThread();
+	bool IsPlaying() const { return bPlaying; }
+
 private:
+	class FWebMBackgroundReader;
+
 	//Theoretically only need 2 buffered textures, but we have extra to avoid needing to make a copy of the AvPlayer data to pass to an RHI thread command.  Instead, we buffer deeper and update the textures on the Render thread.
 	static const int32 NumBufferedTextures = 4;
 
@@ -45,6 +51,8 @@ private:
 	TUniquePtr<FWebMAudioDecoder> AudioDecoder;
 	TUniquePtr<FWebMContainer> Container;
 	TUniquePtr<FWebMAudioBackend> AudioBackend;
+	TUniquePtr<FRunnableThread> BackgroundReaderThread;
+	TUniquePtr<FWebMBackgroundReader> BackgroundReader;
 	TSharedPtr<FMediaSamples, ESPMode::ThreadSafe> Samples;
 	TSharedPtr<FMovieViewport> Viewport;
 	float CurrentTime;
@@ -54,5 +62,4 @@ private:
 	void ReleaseAcquiredResources();
 	bool DisplayFrames(float InDeltaTime);
 	bool SendAudio(float InDeltaTime);
-	bool DecodeMoreFrames();
 };
