WebM: free samples after movie end

diff --git a/Engine/Plugins/Runtime/WebMMoviePlayer/Source/WebMMoviePlayer/Private/WebMMovieStreamer.cpp b/Engine/Plugins/Runtime/WebMMoviePlayer/Source/WebMMoviePlayer/Private/WebMMovieStreamer.cpp
--- a/Engine/Plugins/Runtime/WebMMoviePlayer/Source/WebMMoviePlayer/Private/WebMMovieStreamer.cpp
+++ b/Engine/Plugins/Runtime/WebMMoviePlayer/Source/WebMMoviePlayer/Private/WebMMovieStreamer.cpp
@@ -68,7 +68,6 @@ FWebMMovieStreamer::FWebMMovieStreamer()
 	: CurrentTexture(0)
 	, AudioBackend(MakeUnique<FWebMAudioBackend>())
 	, BackgroundReader(MakeUnique<FWebMBackgroundReader>(*this))
-	, Samples(MakeShareable(new FMediaSamples()))
 	, Viewport(MakeShareable(new FMovieViewport()))
 	, CurrentTime(0)
 	, bPlaying(false)
@@ -93,6 +92,7 @@ void FWebMMovieStreamer::Cleanup()
 		BackgroundReaderThread.Reset();
 	}
 
+	Samples.Reset();
 	VideoDecoder.Reset();
 	AudioDecoder.Reset();
 	Container.Reset();
@@ -148,6 +148,7 @@ bool FWebMMovieStreamer::StartNextMovie()
 			return false;
 		}
 
+		Samples = MakeShareable(new FMediaSamples());
 		AudioDecoder.Reset(new FWebMAudioDecoder(Samples));
 		VideoDecoder.Reset(new FWebMVideoDecoder(Samples));
 
