WebM video streamer: better resource management, fixed Vulkan crash.

diff --git a/Engine/Plugins/Media/WebMMedia/Source/WebMMedia/Private/Player/WebMMediaTextureSample.cpp b/Engine/Plugins/Media/WebMMedia/Source/WebMMedia/Private/Player/WebMMediaTextureSample.cpp
--- a/Engine/Plugins/Media/WebMMedia/Source/WebMMedia/Private/Player/WebMMediaTextureSample.cpp
+++ b/Engine/Plugins/Media/WebMMedia/Source/WebMMedia/Private/Player/WebMMediaTextureSample.cpp
@@ -95,6 +95,9 @@ bool FWebMMediaTextureSample::IsOutputSrgb() const
 
 void FWebMMediaTextureSample::ShutdownPoolable()
 {
+	// Drop reference to the texture. It should be released by the outside system.
+	Texture = nullptr;
+	
 	Time = FTimespan::Zero();
 }
 
diff --git a/Engine/Plugins/Runtime/WebMMoviePlayer/Source/WebMMoviePlayer/Private/WebMMovieStreamer.cpp b/Engine/Plugins/Runtime/WebMMoviePlayer/Source/WebMMoviePlayer/Private/WebMMovieStreamer.cpp
--- a/Engine/Plugins/Runtime/WebMMoviePlayer/Source/WebMMoviePlayer/Private/WebMMovieStreamer.cpp
+++ b/Engine/Plugins/Runtime/WebMMoviePlayer/Source/WebMMoviePlayer/Private/WebMMovieStreamer.cpp
@@ -43,7 +43,10 @@ public:
 
 		while (!bIsFinished && Streamer.IsPlaying())
 		{
-			bIsFinished = Streamer.DecodeMoreFramesInAnotherThread();
+			if (Streamer.GetFramesCurrentlyProcessing() < 30)
+			{
+				bIsFinished = Streamer.DecodeMoreFramesInAnotherThread();
+			}
 		}
 
 		return 0;
@@ -65,10 +68,10 @@ private:
 };
 
 FWebMMovieStreamer::FWebMMovieStreamer()
-	: CurrentTexture(0)
-	, AudioBackend(MakeUnique<FWebMAudioBackend>())
+	: AudioBackend(MakeUnique<FWebMAudioBackend>())
 	, BackgroundReader(MakeUnique<FWebMBackgroundReader>(*this))
 	, Viewport(MakeShareable(new FMovieViewport()))
+	, FramesCurrentlyProcessing(0)
 	, CurrentTime(0)
 	, bPlaying(false)
 {
@@ -104,7 +107,7 @@ void FWebMMovieStreamer::Cleanup()
 
 FTexture2DRHIRef FWebMMovieStreamer::GetTexture()
 {
-	return BufferedVideoTextures[CurrentTexture].IsValid() ? BufferedVideoTextures[CurrentTexture]->GetRHIRef() : nullptr;
+	return SlateVideoTexture.IsValid() ? SlateVideoTexture->GetRHIRef() : nullptr;
 }
 
 bool FWebMMovieStreamer::Init(const TArray<FString>& InMoviePaths, TEnumAsByte<EMoviePlaybackType> InPlaybackType)
@@ -238,43 +241,35 @@ void FWebMMovieStreamer::ForceCompletion()
 
 void FWebMMovieStreamer::ReleaseAcquiredResources()
 {
-	// NOTE: Called from the main thread
-	for (int32 i = 0; i < NumBufferedTextures; ++i)
-	{
-		TSharedPtr<FSlateTexture2DRHIRef, ESPMode::ThreadSafe>& VideoTexture = BufferedVideoTextures[i];
-		if (VideoTexture.IsValid())
-		{
-			// Schedule the release of the video texture(s) and wait the release to complete before
-			// losing the reference to the texture
-			BeginReleaseResource(VideoTexture.Get());
-			FlushRenderingCommands();
-			VideoTexture.Reset();
-		}
-	}
-
+	SlateVideoTexture.Reset();
 	Viewport->SetTexture(nullptr);
 }
 
 bool FWebMMovieStreamer::DisplayFrames(float InDeltaTime)
 {
+	if (!SlateVideoTexture.IsValid())
+	{
+		SlateVideoTexture = MakeShareable(new FSlateTexture2DRHIRef(nullptr, 0, 0));
+	}
+
 	TSharedPtr<IMediaTextureSample, ESPMode::ThreadSafe> VideoSample;
 	TRange<FTimespan> TimeRange(FTimespan::Zero(), FTimespan::FromSeconds(CurrentTime));
 	bool bFoundSample = Samples->FetchVideo(TimeRange, VideoSample);
 
 	if (bFoundSample && VideoSample)
 	{
-		CurrentTexture = (CurrentTexture + 1) % NumBufferedTextures;
+		FPlatformAtomics::InterlockedDecrement(&FramesCurrentlyProcessing);
 
 		FWebMMediaTextureSample* WebMSample = StaticCast<FWebMMediaTextureSample*>(VideoSample.Get());
 
-		if (!BufferedVideoTextures[CurrentTexture].IsValid())
+		if (SlateVideoTexture->IsValid())
 		{
-			BufferedVideoTextures[CurrentTexture] = MakeShareable(new FSlateTexture2DRHIRef(nullptr, 0, 0));
+			SlateVideoTexture->ReleaseDynamicRHI();
 		}
 
-		BufferedVideoTextures[CurrentTexture]->SetRHIRef(WebMSample->GetTextureRef(), WebMSample->GetDim().X, WebMSample->GetDim().Y);
+		SlateVideoTexture->SetRHIRef(WebMSample->GetTextureRef(), WebMSample->GetDim().X, WebMSample->GetDim().Y);
 
-		Viewport->SetTexture(BufferedVideoTextures[CurrentTexture]);
+		Viewport->SetTexture(SlateVideoTexture);
 	}
 
 	return Samples->NumVideoSamples() > 0 || VideoDecoder->IsBusy();
@@ -296,7 +291,6 @@ bool FWebMMovieStreamer::SendAudio(float InDeltaTime)
 		bFoundSample = Samples->FetchAudio(TimeRange, AudioSample);
 	}
 
-
 	return Samples->NumAudio() > 0 || AudioDecoder->IsBusy();
 }
 
@@ -312,6 +306,8 @@ bool FWebMMovieStreamer::DecodeMoreFramesInAnotherThread()
 	// Trigger video decoding
 	if (VideoFrames.Num() > 0)
 	{
+		FPlatformAtomics::InterlockedAdd(&FramesCurrentlyProcessing, VideoFrames.Num());
+
 		VideoDecoder->DecodeVideoFramesAsync(VideoFrames);
 	}
 
diff --git a/Engine/Plugins/Runtime/WebMMoviePlayer/Source/WebMMoviePlayer/Private/WebMMovieStreamer.h b/Engine/Plugins/Runtime/WebMMoviePlayer/Source/WebMMoviePlayer/Private/WebMMovieStreamer.h
--- a/Engine/Plugins/Runtime/WebMMoviePlayer/Source/WebMMoviePlayer/Private/WebMMovieStreamer.h
+++ b/Engine/Plugins/Runtime/WebMMoviePlayer/Source/WebMMoviePlayer/Private/WebMMovieStreamer.h
@@ -33,18 +33,10 @@ public:
 
 	bool DecodeMoreFramesInAnotherThread();
 	bool IsPlaying() const { return bPlaying; }
+	int32 GetFramesCurrentlyProcessing() { return FramesCurrentlyProcessing; }
 
 private:
 	class FWebMBackgroundReader;
-
-	//Theoretically only need 2 buffered textures, but we have extra to avoid needing to make a copy of the AvPlayer data to pass to an RHI thread command.  Instead, we buffer deeper and update the textures on the Render thread.
-	static const int32 NumBufferedTextures = 4;
-
-	/** Texture and viewport data for displaying to Slate.  SlateVideoTexture is always used by the viewport, but it's texture reference is swapped out when a new frame is available.
-	 Method assumes 1 Tick() call per frame, and that the Streamer Tick comes before Slate rendering */
-	TSharedPtr<FSlateTexture2DRHIRef, ESPMode::ThreadSafe> BufferedVideoTextures[NumBufferedTextures];
-	int32 CurrentTexture;
-
 	TArray<FString> MovieQueue;
 	FString MovieName;
 	TUniquePtr<FWebMVideoDecoder> VideoDecoder;
@@ -55,6 +47,8 @@ private:
 	TUniquePtr<FWebMBackgroundReader> BackgroundReader;
 	TSharedPtr<FMediaSamples, ESPMode::ThreadSafe> Samples;
 	TSharedPtr<FMovieViewport> Viewport;
+	TSharedPtr<FSlateTexture2DRHIRef, ESPMode::ThreadSafe> SlateVideoTexture;
+	int32 FramesCurrentlyProcessing;
 	float CurrentTime;
 	bool bPlaying;
 
