WebM media streamer: removed background thread, split processing of audio and video

diff --git a/Engine/Plugins/Runtime/WebMMoviePlayer/Source/WebMMoviePlayer/Private/WebMMovieStreamer.cpp b/Engine/Plugins/Runtime/WebMMoviePlayer/Source/WebMMoviePlayer/Private/WebMMovieStreamer.cpp
--- a/Engine/Plugins/Runtime/WebMMoviePlayer/Source/WebMMoviePlayer/Private/WebMMovieStreamer.cpp
+++ b/Engine/Plugins/Runtime/WebMMoviePlayer/Source/WebMMoviePlayer/Private/WebMMovieStreamer.cpp
@@ -11,67 +11,16 @@
 #include "RHIStaticStates.h"
 #include "WebMVideoDecoder.h"
 #include "WebMAudioDecoder.h"
-#include "WebMMediaFrame.h"
 #include "WebMContainer.h"
 #include "WebMMediaAudioSample.h"
 #include "WebMMediaTextureSample.h"
-#include "GenericPlatform/GenericPlatformProcess.h"
-#include "HAL/Runnable.h"
-#include "HAL/RunnableThread.h"
-#include "Misc/SingleThreadRunnable.h"
 
 DEFINE_LOG_CATEGORY(LogWebMMoviePlayer);
 
-class FWebMMovieStreamer::FWebMBackgroundReader : public FRunnable, private FSingleThreadRunnable
-{
-public:
-	FWebMBackgroundReader(FWebMMovieStreamer& InStreamer)
-		: Streamer(InStreamer)
-		, bIsFinished(false)
-	{
-	}
-
-	//~ FRunnable interface
-	virtual FSingleThreadRunnable* GetSingleThreadInterface() override
-	{
-		return this;
-	}
-
-	virtual uint32 Run() override
-	{
-		bIsFinished = false;
-
-		while (!bIsFinished && Streamer.IsPlaying())
-		{
-			if (Streamer.GetFramesCurrentlyProcessing() < 30)
-			{
-				bIsFinished = Streamer.DecodeMoreFramesInAnotherThread();
-			}
-		}
-
-		return 0;
-	}
-
-protected:
-	//~ FSingleThreadRunnable interface
-	virtual void Tick() override
-	{
-		if (!bIsFinished && Streamer.IsPlaying())
-		{
-			bIsFinished = Streamer.DecodeMoreFramesInAnotherThread();
-		}
-	}
-
-private:
-	FWebMMovieStreamer& Streamer;
-	bool bIsFinished;
-};
-
 FWebMMovieStreamer::FWebMMovieStreamer()
 	: AudioBackend(MakeUnique<FWebMAudioBackend>())
-	, BackgroundReader(MakeUnique<FWebMBackgroundReader>(*this))
 	, Viewport(MakeShareable(new FMovieViewport()))
-	, FramesCurrentlyProcessing(0)
+	, VideoFramesCurrentlyProcessing(0)
 	, CurrentTime(0)
 	, bPlaying(false)
 {
@@ -88,13 +37,7 @@ FWebMMovieStreamer::~FWebMMovieStreamer()
 void FWebMMovieStreamer::Cleanup()
 {
 	bPlaying = false;
-
-	if (BackgroundReaderThread)
-	{
-		BackgroundReaderThread->WaitForCompletion();
-		BackgroundReaderThread.Reset();
-	}
-
+	
 	Samples.Reset();
 	VideoDecoder.Reset();
 	AudioDecoder.Reset();
@@ -169,8 +112,6 @@ bool FWebMMovieStreamer::StartNextMovie()
 		CurrentTime = 0;
 		bPlaying = true;
 
-		BackgroundReaderThread.Reset(FRunnableThread::Create(BackgroundReader.Get(), TEXT("FWebMBackgroundReader"), 0, TPri_BelowNormal));
-
 		return true;
 	}
 	else
@@ -199,6 +140,8 @@ bool FWebMMovieStreamer::Tick(float InDeltaTime)
 		bHaveThingsToDo |= DisplayFrames(InDeltaTime);
 		bHaveThingsToDo |= SendAudio(InDeltaTime);
 
+		bHaveThingsToDo |= ReadMoreFrames();
+
 		CurrentTime += InDeltaTime;
 
 		if (bHaveThingsToDo)
@@ -258,7 +201,7 @@ bool FWebMMovieStreamer::DisplayFrames(float InDeltaTime)
 
 	if (bFoundSample && VideoSample)
 	{
-		FPlatformAtomics::InterlockedDecrement(&FramesCurrentlyProcessing);
+		VideoFramesCurrentlyProcessing--;
 
 		FWebMMediaTextureSample* WebMSample = StaticCast<FWebMMediaTextureSample*>(VideoSample.Get());
 
@@ -294,21 +237,25 @@ bool FWebMMovieStreamer::SendAudio(float InDeltaTime)
 	return Samples->NumAudio() > 0 || AudioDecoder->IsBusy();
 }
 
-bool FWebMMovieStreamer::DecodeMoreFramesInAnotherThread()
+bool FWebMMovieStreamer::ReadMoreFrames()
 {
-	FTimespan ReadBufferLength = FTimespan::FromSeconds(0.1);
+	FTimespan ReadBufferLength = FTimespan::FromSeconds(0.5);
 
 	TArray<TSharedPtr<FWebMFrame>> AudioFrames;
 	TArray<TSharedPtr<FWebMFrame>> VideoFrames;
 
 	Container->ReadFrames(ReadBufferLength, AudioFrames, VideoFrames);
 
+	VideoFramesToDecodeLater.Enqueue(VideoFrames);
+
 	// Trigger video decoding
-	if (VideoFrames.Num() > 0)
+	while (!VideoFramesToDecodeLater.IsEmpty() && VideoFramesCurrentlyProcessing < 30)
 	{
-		FPlatformAtomics::InterlockedAdd(&FramesCurrentlyProcessing, VideoFrames.Num());
-
-		VideoDecoder->DecodeVideoFramesAsync(VideoFrames);
+		if (VideoFramesToDecodeLater.Dequeue(VideoFrames))
+		{
+			VideoFramesCurrentlyProcessing += VideoFrames.Num();
+			VideoDecoder->DecodeVideoFramesAsync(VideoFrames);
+		}
 	}
 
 	// Trigger audio decoding
@@ -317,5 +264,5 @@ bool FWebMMovieStreamer::DecodeMoreFramesInAnotherThread()
 		AudioDecoder->DecodeAudioFramesAsync(AudioFrames);
 	}
 
-	return VideoFrames.Num() == 0 && AudioFrames.Num() == 0;
+	return VideoFrames.Num() > 0 || AudioFrames.Num() > 0;
 }
diff --git a/Engine/Plugins/Runtime/WebMMoviePlayer/Source/WebMMoviePlayer/Private/WebMMovieStreamer.h b/Engine/Plugins/Runtime/WebMMoviePlayer/Source/WebMMoviePlayer/Private/WebMMovieStreamer.h
--- a/Engine/Plugins/Runtime/WebMMoviePlayer/Source/WebMMoviePlayer/Private/WebMMovieStreamer.h
+++ b/Engine/Plugins/Runtime/WebMMoviePlayer/Source/WebMMoviePlayer/Private/WebMMovieStreamer.h
@@ -4,12 +4,13 @@
 
 #include "MoviePlayer.h"
 #include "WebMAudioBackend.h"
+#include "WebMMediaFrame.h"
+#include "Containers/Queue.h"
 
 class FWebMVideoDecoder;
 class FWebMAudioDecoder;
 class FMediaSamples;
 class FWebMContainer;
-class FRunnableThread;
 
 class FWebMMovieStreamer : public IMovieStreamer
 {
@@ -31,24 +32,18 @@ public:
 	FOnCurrentMovieClipFinished OnCurrentMovieClipFinishedDelegate;
 	virtual FOnCurrentMovieClipFinished& OnCurrentMovieClipFinished() override { return OnCurrentMovieClipFinishedDelegate; }
 
-	bool DecodeMoreFramesInAnotherThread();
-	bool IsPlaying() const { return bPlaying; }
-	int32 GetFramesCurrentlyProcessing() { return FramesCurrentlyProcessing; }
-
 private:
-	class FWebMBackgroundReader;
 	TArray<FString> MovieQueue;
+	TQueue<TArray<TSharedPtr<FWebMFrame>>> VideoFramesToDecodeLater;
 	FString MovieName;
 	TUniquePtr<FWebMVideoDecoder> VideoDecoder;
 	TUniquePtr<FWebMAudioDecoder> AudioDecoder;
 	TUniquePtr<FWebMContainer> Container;
 	TUniquePtr<FWebMAudioBackend> AudioBackend;
-	TUniquePtr<FRunnableThread> BackgroundReaderThread;
-	TUniquePtr<FWebMBackgroundReader> BackgroundReader;
 	TSharedPtr<FMediaSamples, ESPMode::ThreadSafe> Samples;
 	TSharedPtr<FMovieViewport> Viewport;
 	TSharedPtr<FSlateTexture2DRHIRef, ESPMode::ThreadSafe> SlateVideoTexture;
-	int32 FramesCurrentlyProcessing;
+	int32 VideoFramesCurrentlyProcessing;
 	float CurrentTime;
 	bool bPlaying;
 
@@ -56,4 +51,5 @@ private:
 	void ReleaseAcquiredResources();
 	bool DisplayFrames(float InDeltaTime);
 	bool SendAudio(float InDeltaTime);
+	bool ReadMoreFrames();
 };
