WebM: streamer timing fix

diff --git a/Engine/Plugins/Runtime/WebMMoviePlayer/Source/WebMMoviePlayer/Private/WebMMovieStreamer.cpp b/Engine/Plugins/Runtime/WebMMoviePlayer/Source/WebMMoviePlayer/Private/WebMMovieStreamer.cpp
--- a/Engine/Plugins/Runtime/WebMMoviePlayer/Source/WebMMoviePlayer/Private/WebMMovieStreamer.cpp
+++ b/Engine/Plugins/Runtime/WebMMoviePlayer/Source/WebMMoviePlayer/Private/WebMMovieStreamer.cpp
@@ -21,7 +21,7 @@ FWebMMovieStreamer::FWebMMovieStreamer()
 	: AudioBackend(MakeUnique<FWebMAudioBackend>())
 	, Viewport(MakeShareable(new FMovieViewport()))
 	, VideoFramesCurrentlyProcessing(0)
-	, CurrentTime(0)
+	, StartTime(0)
 	, bPlaying(false)
 {
 	AudioBackend->InitializePlatform();
@@ -109,7 +109,7 @@ bool FWebMMovieStreamer::StartNextMovie()
 
 		AudioBackend->StartStreaming(DefaultAudioTrack.SampleRate, DefaultAudioTrack.NumOfChannels);
 
-		CurrentTime = 0;
+		StartTime = FPlatformTime::Seconds();
 		bPlaying = true;
 
 		return true;
@@ -142,8 +142,6 @@ bool FWebMMovieStreamer::Tick(float InDeltaTime)
 
 		bHaveThingsToDo |= ReadMoreFrames();
 
-		CurrentTime += InDeltaTime;
-
 		if (bHaveThingsToDo)
 		{
 			// We're still playing this movie
@@ -195,8 +193,11 @@ bool FWebMMovieStreamer::DisplayFrames(float InDeltaTime)
 		SlateVideoTexture = MakeShareable(new FSlateTexture2DRHIRef(nullptr, 0, 0));
 	}
 
+	double CurrentTime = FPlatformTime::Seconds();
+	double MovieTime = CurrentTime - StartTime;
+
 	TSharedPtr<IMediaTextureSample, ESPMode::ThreadSafe> VideoSample;
-	TRange<FTimespan> TimeRange(FTimespan::Zero(), FTimespan::FromSeconds(CurrentTime));
+	TRange<FTimespan> TimeRange(FTimespan::Zero(), FTimespan::FromSeconds(MovieTime));
 	bool bFoundSample = Samples->FetchVideo(TimeRange, VideoSample);
 
 	if (bFoundSample && VideoSample)
@@ -239,7 +240,7 @@ bool FWebMMovieStreamer::SendAudio(float InDeltaTime)
 
 bool FWebMMovieStreamer::ReadMoreFrames()
 {
-	FTimespan ReadBufferLength = FTimespan::FromSeconds(0.5);
+	FTimespan ReadBufferLength = FTimespan::FromSeconds(1.0 / 30.0);
 
 	TArray<TSharedPtr<FWebMFrame>> AudioFrames;
 	TArray<TSharedPtr<FWebMFrame>> VideoFrames;
diff --git a/Engine/Plugins/Runtime/WebMMoviePlayer/Source/WebMMoviePlayer/Private/WebMMovieStreamer.h b/Engine/Plugins/Runtime/WebMMoviePlayer/Source/WebMMoviePlayer/Private/WebMMovieStreamer.h
--- a/Engine/Plugins/Runtime/WebMMoviePlayer/Source/WebMMoviePlayer/Private/WebMMovieStreamer.h
+++ b/Engine/Plugins/Runtime/WebMMoviePlayer/Source/WebMMoviePlayer/Private/WebMMovieStreamer.h
@@ -44,7 +44,7 @@ private:
 	TSharedPtr<FMovieViewport> Viewport;
 	TSharedPtr<FSlateTexture2DRHIRef, ESPMode::ThreadSafe> SlateVideoTexture;
 	int32 VideoFramesCurrentlyProcessing;
-	float CurrentTime;
+	double StartTime;
 	bool bPlaying;
 
 	bool StartNextMovie();
