More correct SDL initialization/deinitialization flow.

- Attempts to alleviate UE-65544

diff --git a/Engine/Plugins/Runtime/WebMMoviePlayer/Source/WebMMoviePlayer/Private/Audio/Unix/WebMAudioBackendSDL.cpp b/Engine/Plugins/Runtime/WebMMoviePlayer/Source/WebMMoviePlayer/Private/Audio/Unix/WebMAudioBackendSDL.cpp
--- a/Engine/Plugins/Runtime/WebMMoviePlayer/Source/WebMMoviePlayer/Private/Audio/Unix/WebMAudioBackendSDL.cpp
+++ b/Engine/Plugins/Runtime/WebMMoviePlayer/Source/WebMMoviePlayer/Private/Audio/Unix/WebMAudioBackendSDL.cpp
@@ -36,7 +36,9 @@ void FWebMAudioBackendSDL::ShutdownPlatform()
 
 	if (bSDLInitialized)
 	{
-		SDL_Quit();
+		// this is refcounted
+		SDL_QuitSubSystem(SDL_INIT_AUDIO);
+		bSDLInitialized = false;
 	}
 }
 
diff --git a/Engine/Plugins/Runtime/WebMMoviePlayer/Source/WebMMoviePlayer/Private/WebMMovieStreamer.cpp b/Engine/Plugins/Runtime/WebMMoviePlayer/Source/WebMMoviePlayer/Private/WebMMovieStreamer.cpp
--- a/Engine/Plugins/Runtime/WebMMoviePlayer/Source/WebMMoviePlayer/Private/WebMMovieStreamer.cpp
+++ b/Engine/Plugins/Runtime/WebMMoviePlayer/Source/WebMMoviePlayer/Private/WebMMovieStreamer.cpp
@@ -24,28 +24,19 @@ FWebMMovieStreamer::FWebMMovieStreamer()
 	, StartTime(0)
 	, bPlaying(false)
 {
-	AudioBackend->InitializePlatform();
 }
 
 FWebMMovieStreamer::~FWebMMovieStreamer()
 {
 	Cleanup();
-
-	AudioBackend->ShutdownPlatform();
 }
 
 void FWebMMovieStreamer::Cleanup()
 {
 	bPlaying = false;
-	
-	Samples.Reset();
-	VideoDecoder.Reset();
-	AudioDecoder.Reset();
-	Container.Reset();
 
 	ReleaseAcquiredResources();
-
-	AudioBackend->StopStreaming();
+	AudioBackend->ShutdownPlatform();
 }
 
 FTexture2DRHIRef FWebMMovieStreamer::GetTexture()
@@ -57,6 +48,7 @@ bool FWebMMovieStreamer::Init(const TArray<FString>& InMoviePaths, TEnumAsByte<E
 {
 	// Initializes the streamer for audio and video playback of the given path(s).
 	// NOTE: If multiple paths are provided, it is expect that they be played back seamlessly.
+	AudioBackend->InitializePlatform();
 
 	// Add the given paths to the movie queue
 	MovieQueue.Append(InMoviePaths);
@@ -69,7 +61,7 @@ bool FWebMMovieStreamer::StartNextMovie()
 {
 	if (MovieQueue.Num() > 0)
 	{
-		Cleanup();
+		ReleaseAcquiredResources();
 
 		MovieName = MovieQueue[0];
 
@@ -182,8 +174,15 @@ void FWebMMovieStreamer::ForceCompletion()
 
 void FWebMMovieStreamer::ReleaseAcquiredResources()
 {
+	Samples.Reset();
+	VideoDecoder.Reset();
+	AudioDecoder.Reset();
+	Container.Reset();
+
 	SlateVideoTexture.Reset();
 	Viewport->SetTexture(nullptr);
+
+	AudioBackend->StopStreaming();
 }
 
 bool FWebMMovieStreamer::DisplayFrames(float InDeltaTime)
diff --git a/Engine/Source/Runtime/Linux/AudioMixerSDL/Private/AudioMixerPlatformSDL.cpp b/Engine/Source/Runtime/Linux/AudioMixerSDL/Private/AudioMixerPlatformSDL.cpp
--- a/Engine/Source/Runtime/Linux/AudioMixerSDL/Private/AudioMixerPlatformSDL.cpp
+++ b/Engine/Source/Runtime/Linux/AudioMixerSDL/Private/AudioMixerPlatformSDL.cpp
@@ -72,6 +72,8 @@ namespace Audio
 		StopAudioStream();
 		CloseAudioStream();
 
+		// this is refcounted
+		SDL_QuitSubSystem(SDL_INIT_AUDIO);
 		bInitialized = false;
 
 		return true;
