WebMMediaStreamer: reduced number of buffered frames to reduce texture memory consuption

diff --git a/Engine/Plugins/Runtime/WebMMoviePlayer/Source/WebMMoviePlayer/Private/WebMMovieStreamer.cpp b/Engine/Plugins/Runtime/WebMMoviePlayer/Source/WebMMoviePlayer/Private/WebMMovieStreamer.cpp
--- a/Engine/Plugins/Runtime/WebMMoviePlayer/Source/WebMMoviePlayer/Private/WebMMovieStreamer.cpp
+++ b/Engine/Plugins/Runtime/WebMMoviePlayer/Source/WebMMoviePlayer/Private/WebMMovieStreamer.cpp
@@ -249,7 +249,7 @@ bool FWebMMovieStreamer::ReadMoreFrames()
 	VideoFramesToDecodeLater.Enqueue(VideoFrames);
 
 	// Trigger video decoding
-	while (!VideoFramesToDecodeLater.IsEmpty() && VideoFramesCurrentlyProcessing < 30)
+	while (!VideoFramesToDecodeLater.IsEmpty() && VideoFramesCurrentlyProcessing < 5)
 	{
 		if (VideoFramesToDecodeLater.Dequeue(VideoFrames))
 		{
