WebM: error reporting fix / crash fix

GitHub PR #5541

diff --git a/Engine/Plugins/Runtime/WebMMoviePlayer/Source/WebMMoviePlayer/Private/WebMMovieStreamer.cpp b/Engine/Plugins/Runtime/WebMMoviePlayer/Source/WebMMoviePlayer/Private/WebMMovieStreamer.cpp
--- a/Engine/Plugins/Runtime/WebMMoviePlayer/Source/WebMMoviePlayer/Private/WebMMovieStreamer.cpp
+++ b/Engine/Plugins/Runtime/WebMMoviePlayer/Source/WebMMoviePlayer/Private/WebMMovieStreamer.cpp
@@ -79,7 +79,10 @@ bool FWebMMovieStreamer::StartNextMovie()
 		}
 		else
 		{
-			UE_LOG(LogWebMMoviePlayer, Error, TEXT("Movie '%s' not found."));
+			UE_LOG(LogWebMMoviePlayer, Error, TEXT("Movie '%s' not found."), *MoviePath);
+			
+			MovieName = FString();
+			return false;
 		}
 
 		UE_LOG(LogWebMMoviePlayer, Verbose, TEXT("Starting '%s'"), *MoviePath);
