Fix startup movies on Windows. WebM player was taking over startup movie playback.

diff --git a/Engine/Plugins/Runtime/WebMMoviePlayer/Source/WebMMoviePlayer/Private/WebMMoviePlayer.cpp b/Engine/Plugins/Runtime/WebMMoviePlayer/Source/WebMMoviePlayer/Private/WebMMoviePlayer.cpp
--- a/Engine/Plugins/Runtime/WebMMoviePlayer/Source/WebMMoviePlayer/Private/WebMMoviePlayer.cpp
+++ b/Engine/Plugins/Runtime/WebMMoviePlayer/Source/WebMMoviePlayer/Private/WebMMoviePlayer.cpp
@@ -5,27 +5,33 @@
 #include "MoviePlayer.h"
 #include "Modules/ModuleManager.h"
 
-#if WITH_WEBM_LIBS
+#if WITH_WEBM_LIBS && !PLATFORM_WINDOWS
+	#define WITH_WEBM_STARTUP_MOVIES 1
+#else
+	#define WITH_WEBM_STARTUP_MOVIES 0
+#endif
+
+#if WITH_WEBM_STARTUP_MOVIES
 TSharedPtr<FWebMMovieStreamer> WebMMovieStreamer;
-#endif // WITH_WEBM_LIBS
+#endif // WITH_WEBM_STARTUP_MOVIES
 
 class FWebMMoviePlayerModule : public IModuleInterface
 {
 	/** IModuleInterface implementation */
 	virtual void StartupModule() override
 	{
-#if WITH_WEBM_LIBS
+#if WITH_WEBM_STARTUP_MOVIES
 		FWebMMovieStreamer *Streamer = new FWebMMovieStreamer;
 		WebMMovieStreamer = MakeShareable(Streamer);
 		GetMoviePlayer()->RegisterMovieStreamer(WebMMovieStreamer);
-#endif // WITH_WEBM_LIBS
+#endif // WITH_WEBM_STARTUP_MOVIES
 	}
 
 	virtual void ShutdownModule() override
 	{
-#if WITH_WEBM_LIBS
+#if WITH_WEBM_STARTUP_MOVIES
 		WebMMovieStreamer.Reset();
-#endif // WITH_WEBM_LIBS
+#endif // WITH_WEBM_STARTUP_MOVIES
 	}
 };
 
