Allow for platforms to override the default device name for WebMAudioBackend

diff --git a/Engine/Plugins/Runtime/WebMMoviePlayer/Source/WebMMoviePlayer/Private/Audio/Linux/WebMAudioBackendDefines.h b/Engine/Plugins/Runtime/WebMMoviePlayer/Source/WebMMoviePlayer/Private/Audio/Linux/WebMAudioBackendDefines.h
new file mode 100644
--- /dev/null
+++ b/Engine/Plugins/Runtime/WebMMoviePlayer/Source/WebMMoviePlayer/Private/Audio/Linux/WebMAudioBackendDefines.h
@@ -0,0 +1,7 @@
+// Copyright 1998-2019 Epic Games, Inc. All Rights Reserved.
+
+#pragma once
+
+#include "WebMAudioBackend.h"
+
+using FWebMAudioBackend = FWebMAudioBackendSDL;
diff --git a/Engine/Plugins/Runtime/WebMMoviePlayer/Source/WebMMoviePlayer/Private/Audio/Unix/WebMAudioBackend.h b/Engine/Plugins/Runtime/WebMMoviePlayer/Source/WebMMoviePlayer/Private/Audio/Unix/WebMAudioBackend.h
--- a/Engine/Plugins/Runtime/WebMMoviePlayer/Source/WebMMoviePlayer/Private/Audio/Unix/WebMAudioBackend.h
+++ b/Engine/Plugins/Runtime/WebMMoviePlayer/Source/WebMMoviePlayer/Private/Audio/Unix/WebMAudioBackend.h
@@ -22,9 +22,9 @@ public:
 	virtual void StopStreaming();
 	virtual bool SendAudio(const uint8* Buffer, size_t BufferSize);
 
+	virtual FString GetDefaultDeviceName();
+
 private:
 	SDL_AudioDeviceID AudioDevice;
 	bool bSDLInitialized;
 };
-
-using FWebMAudioBackend = FWebMAudioBackendSDL;
diff --git a/Engine/Plugins/Runtime/WebMMoviePlayer/Source/WebMMoviePlayer/Private/Audio/Unix/WebMAudioBackendSDL.cpp b/Engine/Plugins/Runtime/WebMMoviePlayer/Source/WebMMoviePlayer/Private/Audio/Unix/WebMAudioBackendSDL.cpp
--- a/Engine/Plugins/Runtime/WebMMoviePlayer/Source/WebMMoviePlayer/Private/Audio/Unix/WebMAudioBackendSDL.cpp
+++ b/Engine/Plugins/Runtime/WebMMoviePlayer/Source/WebMMoviePlayer/Private/Audio/Unix/WebMAudioBackendSDL.cpp
@@ -56,7 +56,9 @@ bool FWebMAudioBackendSDL::StartStreaming(int32 SampleRate, int32 NumOfChannels)
 	AudioSpec.channels = NumOfChannels;
 	AudioSpec.samples = 4096;
 
-	AudioDevice = SDL_OpenAudioDevice(nullptr, 0, &AudioSpec, nullptr, 0);
+	FString DefaultDeivceName = GetDefaultDeviceName();
+
+	AudioDevice = SDL_OpenAudioDevice(TCHAR_TO_ANSI(*DefaultDeivceName), 0, &AudioSpec, nullptr, 0);
 
 	if (!AudioDevice)
 	{
@@ -92,3 +94,8 @@ bool FWebMAudioBackendSDL::SendAudio(const uint8* Buffer, size_t BufferSize)
 		return true;
 	}
 }
+
+FString FWebMAudioBackendSDL::GetDefaultDeviceName()
+{
+	return {};
+}
diff --git a/Engine/Plugins/Runtime/WebMMoviePlayer/Source/WebMMoviePlayer/Private/WebMMovieStreamer.h b/Engine/Plugins/Runtime/WebMMoviePlayer/Source/WebMMoviePlayer/Private/WebMMovieStreamer.h
--- a/Engine/Plugins/Runtime/WebMMoviePlayer/Source/WebMMoviePlayer/Private/WebMMovieStreamer.h
+++ b/Engine/Plugins/Runtime/WebMMoviePlayer/Source/WebMMoviePlayer/Private/WebMMovieStreamer.h
@@ -6,6 +6,7 @@
 
 #include "MoviePlayer.h"
 #include "WebMAudioBackend.h"
+#include "WebMAudioBackendDefines.h"
 #include "WebMMediaFrame.h"
 #include "Containers/Queue.h"
 #include "WebMSamplesSink.h"
diff --git a/Engine/Plugins/Runtime/WebMMoviePlayer/Source/WebMMoviePlayer/WebMMoviePlayer.Build.cs b/Engine/Plugins/Runtime/WebMMoviePlayer/Source/WebMMoviePlayer/WebMMoviePlayer.Build.cs
--- a/Engine/Plugins/Runtime/WebMMoviePlayer/Source/WebMMoviePlayer/WebMMoviePlayer.Build.cs
+++ b/Engine/Plugins/Runtime/WebMMoviePlayer/Source/WebMMoviePlayer/WebMMoviePlayer.Build.cs
@@ -27,6 +27,7 @@ namespace UnrealBuildTool.Rules
 				PrivateDependencyModuleNames.Add("SDL2");
 
 				PrivateIncludePaths.Add("WebMMoviePlayer/Private/Audio/Unix");
+				PrivateIncludePaths.Add("WebMMoviePlayer/Private/Audio/" + Target.Platform);
 			}
 			else
 			{
