SDL does not treat empty string as NULL for device name

diff --git a/Engine/Plugins/Runtime/WebMMoviePlayer/Source/WebMMoviePlayer/Private/Audio/Unix/WebMAudioBackendSDL.cpp b/Engine/Plugins/Runtime/WebMMoviePlayer/Source/WebMMoviePlayer/Private/Audio/Unix/WebMAudioBackendSDL.cpp
--- a/Engine/Plugins/Runtime/WebMMoviePlayer/Source/WebMMoviePlayer/Private/Audio/Unix/WebMAudioBackendSDL.cpp
+++ b/Engine/Plugins/Runtime/WebMMoviePlayer/Source/WebMMoviePlayer/Private/Audio/Unix/WebMAudioBackendSDL.cpp
@@ -57,8 +57,14 @@ bool FWebMAudioBackendSDL::StartStreaming(int32 SampleRate, int32 NumOfChannels)
 	AudioSpec.samples = 4096;
 
 	FString DefaultDeivceName = GetDefaultDeviceName();
-
-	AudioDevice = SDL_OpenAudioDevice(TCHAR_TO_ANSI(*DefaultDeivceName), 0, &AudioSpec, nullptr, 0);
+	if (DefaultDeivceName.Len() > 0)
+	{
+		AudioDevice = SDL_OpenAudioDevice(TCHAR_TO_ANSI(*DefaultDeivceName), 0, &AudioSpec, nullptr, 0);
+	}
+	else
+	{
+		AudioDevice = SDL_OpenAudioDevice(nullptr, 0, &AudioSpec, nullptr, 0);
+	}
 
 	if (!AudioDevice)
 	{
