Add support for tick/pause WebMAudioBackend

diff --git a/Engine/Plugins/Runtime/WebMMoviePlayer/Source/WebMMoviePlayer/Private/Audio/Unix/WebMAudioBackend.h b/Engine/Plugins/Runtime/WebMMoviePlayer/Source/WebMMoviePlayer/Private/Audio/Unix/WebMAudioBackend.h
--- a/Engine/Plugins/Runtime/WebMMoviePlayer/Source/WebMMoviePlayer/Private/Audio/Unix/WebMAudioBackend.h
+++ b/Engine/Plugins/Runtime/WebMMoviePlayer/Source/WebMMoviePlayer/Private/Audio/Unix/WebMAudioBackend.h
@@ -22,9 +22,13 @@ public:
 	virtual void StopStreaming();
 	virtual bool SendAudio(const uint8* Buffer, size_t BufferSize);
 
+	virtual bool IsPaused() const;
+	virtual void Tick(float DeltaTime);
 	virtual FString GetDefaultDeviceName();
 
 private:
+	FCriticalSection AudioDeviceMutex;
+
 	SDL_AudioDeviceID AudioDevice;
 	bool bSDLInitialized;
 };
diff --git a/Engine/Plugins/Runtime/WebMMoviePlayer/Source/WebMMoviePlayer/Private/Audio/Unix/WebMAudioBackendSDL.cpp b/Engine/Plugins/Runtime/WebMMoviePlayer/Source/WebMMoviePlayer/Private/Audio/Unix/WebMAudioBackendSDL.cpp
--- a/Engine/Plugins/Runtime/WebMMoviePlayer/Source/WebMMoviePlayer/Private/Audio/Unix/WebMAudioBackendSDL.cpp
+++ b/Engine/Plugins/Runtime/WebMMoviePlayer/Source/WebMMoviePlayer/Private/Audio/Unix/WebMAudioBackendSDL.cpp
@@ -101,6 +101,15 @@ bool FWebMAudioBackendSDL::SendAudio(const uint8* Buffer, size_t BufferSize)
 	}
 }
 
+bool FWebMAudioBackendSDL::IsPaused() const
+{
+	return false;
+}
+
+void FWebMAudioBackendSDL::Tick(float DeltaTime)
+{
+}
+
 FString FWebMAudioBackendSDL::GetDefaultDeviceName()
 {
 	return {};
