Add new functions to the null backend as well

diff --git a/Engine/Plugins/Runtime/WebMMoviePlayer/Source/WebMMoviePlayer/Private/Audio/Null/WebMAudioBackend.h b/Engine/Plugins/Runtime/WebMMoviePlayer/Source/WebMMoviePlayer/Private/Audio/Null/WebMAudioBackend.h
--- a/Engine/Plugins/Runtime/WebMMoviePlayer/Source/WebMMoviePlayer/Private/Audio/Null/WebMAudioBackend.h
+++ b/Engine/Plugins/Runtime/WebMMoviePlayer/Source/WebMMoviePlayer/Private/Audio/Null/WebMAudioBackend.h
@@ -12,4 +12,7 @@ public:
 	bool StartStreaming(int32 SampleRate, int32 NumOfChannels) { return true; }
 	void StopStreaming() {}
 	bool SendAudio(const uint8* Buffer, size_t BufferSize) { return true; }
+
+	bool IsPaused() const { return false; }
+	void Tick() {}
 };
