WebM Null audio backend needs float arg for tick

diff --git a/Engine/Plugins/Runtime/WebMMoviePlayer/Source/WebMMoviePlayer/Private/Audio/Null/WebMAudioBackend.h b/Engine/Plugins/Runtime/WebMMoviePlayer/Source/WebMMoviePlayer/Private/Audio/Null/WebMAudioBackend.h
--- a/Engine/Plugins/Runtime/WebMMoviePlayer/Source/WebMMoviePlayer/Private/Audio/Null/WebMAudioBackend.h
+++ b/Engine/Plugins/Runtime/WebMMoviePlayer/Source/WebMMoviePlayer/Private/Audio/Null/WebMAudioBackend.h
@@ -13,6 +13,6 @@ public:
 	void StopStreaming() {}
 	bool SendAudio(const uint8* Buffer, size_t BufferSize) { return true; }
 
-	bool IsPaused() const { return false; }
+	bool IsPaused(float DeltaTime) const { return false; }
 	void Tick() {}
 };
