WebM Movie Streamer: better handling of background state (UE-78696).

diff --git a/Engine/Plugins/Runtime/WebMMoviePlayer/Source/WebMMoviePlayer/Private/Audio/Unix/WebMAudioBackend.h b/Engine/Plugins/Runtime/WebMMoviePlayer/Source/WebMMoviePlayer/Private/Audio/Unix/WebMAudioBackend.h
--- a/Engine/Plugins/Runtime/WebMMoviePlayer/Source/WebMMoviePlayer/Private/Audio/Unix/WebMAudioBackend.h
+++ b/Engine/Plugins/Runtime/WebMMoviePlayer/Source/WebMMoviePlayer/Private/Audio/Unix/WebMAudioBackend.h
@@ -12,23 +12,32 @@ THIRD_PARTY_INCLUDES_END
 class FWebMAudioBackendSDL
 {
 public:
+
+	enum class EStreamState
+	{
+		NewMovie,
+		ResumeAfterPause
+	};
+
 	FWebMAudioBackendSDL();
 	virtual ~FWebMAudioBackendSDL();
 
 	//~ IWebMAudioBackend interface
 	virtual bool InitializePlatform();
 	virtual void ShutdownPlatform();
-	virtual bool StartStreaming(int32 SampleRate, int32 NumOfChannels);
+	virtual bool StartStreaming(int32 SampleRate, int32 NumOfChannels, EStreamState StreamState);
 	virtual void StopStreaming();
-	virtual bool SendAudio(const uint8* Buffer, size_t BufferSize);
+	virtual bool SendAudio(const FTimespan& Timespan, const uint8* Buffer, size_t BufferSize);
 
+	virtual void Pause(bool bPause);
 	virtual bool IsPaused() const;
 	virtual void Tick(float DeltaTime);
 	virtual FString GetDefaultDeviceName();
 
-private:
-	FCriticalSection AudioDeviceMutex;
-
+protected:
 	SDL_AudioDeviceID AudioDevice;
 	bool bSDLInitialized;
+
+	/** Whether we're pausing playing back sound. */
+	bool bPaused;
 };
diff --git a/Engine/Plugins/Runtime/WebMMoviePlayer/Source/WebMMoviePlayer/Private/Audio/Unix/WebMAudioBackendSDL.cpp b/Engine/Plugins/Runtime/WebMMoviePlayer/Source/WebMMoviePlayer/Private/Audio/Unix/WebMAudioBackendSDL.cpp
--- a/Engine/Plugins/Runtime/WebMMoviePlayer/Source/WebMMoviePlayer/Private/Audio/Unix/WebMAudioBackendSDL.cpp
+++ b/Engine/Plugins/Runtime/WebMMoviePlayer/Source/WebMMoviePlayer/Private/Audio/Unix/WebMAudioBackendSDL.cpp
@@ -6,6 +6,7 @@
 FWebMAudioBackendSDL::FWebMAudioBackendSDL()
 	: AudioDevice(0)
 	, bSDLInitialized(false)
+	, bPaused(false)
 {
 }
 
@@ -42,7 +43,7 @@ void FWebMAudioBackendSDL::ShutdownPlatform()
 	}
 }
 
-bool FWebMAudioBackendSDL::StartStreaming(int32 SampleRate, int32 NumOfChannels)
+bool FWebMAudioBackendSDL::StartStreaming(int32 SampleRate, int32 NumOfChannels, EStreamState StreamState)
 {
 	if (!bSDLInitialized)
 	{
@@ -86,7 +87,7 @@ void FWebMAudioBackendSDL::StopStreaming()
 	}
 }
 
-bool FWebMAudioBackendSDL::SendAudio(const uint8* Buffer, size_t BufferSize)
+bool FWebMAudioBackendSDL::SendAudio(const FTimespan& Timespan, const uint8* Buffer, size_t BufferSize)
 {
 	int32 Result = SDL_QueueAudio(AudioDevice, Buffer, BufferSize);
 	if (Result < 0)
@@ -96,14 +97,23 @@ bool FWebMAudioBackendSDL::SendAudio(const uint8* Buffer, size_t BufferSize)
 	}
 	else
 	{
-		SDL_PauseAudioDevice(AudioDevice, 0);
+		if (!bPaused)
+		{
+			SDL_PauseAudioDevice(AudioDevice, 0);
+		}
 		return true;
 	}
 }
 
+void FWebMAudioBackendSDL::Pause(bool bPause)
+{
+	bPaused = bPause;
+	SDL_PauseAudioDevice(AudioDevice, bPaused ? 1 : 0);
+}
+
 bool FWebMAudioBackendSDL::IsPaused() const
 {
-	return false;
+	return bPaused;
 }
 
 void FWebMAudioBackendSDL::Tick(float DeltaTime)
diff --git a/Engine/Plugins/Runtime/WebMMoviePlayer/Source/WebMMoviePlayer/Private/WebMMovieStreamer.cpp b/Engine/Plugins/Runtime/WebMMoviePlayer/Source/WebMMoviePlayer/Private/WebMMovieStreamer.cpp
--- a/Engine/Plugins/Runtime/WebMMoviePlayer/Source/WebMMoviePlayer/Private/WebMMovieStreamer.cpp
+++ b/Engine/Plugins/Runtime/WebMMoviePlayer/Source/WebMMoviePlayer/Private/WebMMovieStreamer.cpp
@@ -27,6 +27,7 @@ FWebMMovieStreamer::FWebMMovieStreamer()
 	, VideoFramesCurrentlyProcessing(0)
 	, StartTime(0)
 	, bPlaying(false)
+	, bPaused(false)
 	, TicksLeftToWaitPostCompletion(0)
 {
 }
@@ -39,6 +40,22 @@ FWebMMovieStreamer::~FWebMMovieStreamer()
 void FWebMMovieStreamer::Cleanup()
 {
 	bPlaying = false;
+	bPaused = false;
+	VideoFramesCurrentlyProcessing = 0;
+	TicksLeftToWaitPostCompletion = 0;
+
+	// remove delegates if registered
+	if (ResumeHandle.IsValid())
+	{
+		FCoreDelegates::ApplicationHasEnteredForegroundDelegate.Remove(ResumeHandle);
+		ResumeHandle.Reset();
+	}
+
+	if (PauseHandle.IsValid())
+	{
+		FCoreDelegates::ApplicationWillEnterBackgroundDelegate.Remove(PauseHandle);
+		PauseHandle.Reset();
+	}
 
 	ReleaseAcquiredResources();
 	AudioBackend->ShutdownPlatform();
@@ -58,12 +75,24 @@ bool FWebMMovieStreamer::Init(const TArray<FString>& InMoviePaths, TEnumAsByte<E
 	// Add the given paths to the movie queue
 	MovieQueue.Append(InMoviePaths);
 
+	// register delegate if not registered
+	if (!ResumeHandle.IsValid())
+	{
+		ResumeHandle = FCoreDelegates::ApplicationHasEnteredForegroundDelegate.AddRaw(this, &FWebMMovieStreamer::HandleApplicationHasEnteredForeground);
+	}
+	if (!PauseHandle.IsValid())
+	{
+		PauseHandle = FCoreDelegates::ApplicationWillEnterBackgroundDelegate.AddRaw(this, &FWebMMovieStreamer::HandleApplicationWillEnterBackground);
+	}
+
 	// start our first movie playing
 	return StartNextMovie();
 }
 
 bool FWebMMovieStreamer::StartNextMovie()
 {
+	checkf(!bPaused, TEXT("Should not start next movie when paused!"));
+
 	if (MovieQueue.Num() > 0)
 	{
 		ReleaseAcquiredResources();
@@ -86,7 +115,7 @@ bool FWebMMovieStreamer::StartNextMovie()
 			return false;
 		}
 
-		UE_LOG(LogWebMMoviePlayer, Verbose, TEXT("Starting '%s'"), *MoviePath);
+		UE_LOG(LogWebMMoviePlayer, Log, TEXT("Starting '%s'"), *MoviePath);
 
 		if (!Container->Open(MoviePath))
 		{
@@ -107,7 +136,7 @@ bool FWebMMovieStreamer::StartNextMovie()
 		AudioDecoder->Initialize(DefaultAudioTrack.CodecName, DefaultAudioTrack.SampleRate, DefaultAudioTrack.NumOfChannels, DefaultAudioTrack.CodecPrivateData, DefaultAudioTrack.CodecPrivateDataSize);
 		VideoDecoder->Initialize(DefaultVideoTrack.CodecName);
 
-		AudioBackend->StartStreaming(DefaultAudioTrack.SampleRate, DefaultAudioTrack.NumOfChannels);
+		AudioBackend->StartStreaming(DefaultAudioTrack.SampleRate, DefaultAudioTrack.NumOfChannels, FWebMAudioBackendSDL::EStreamState::NewMovie);
 
 		StartTime = FPlatformTime::Seconds();
 		bPlaying = true;
@@ -116,7 +145,7 @@ bool FWebMMovieStreamer::StartNextMovie()
 	}
 	else
 	{
-		UE_LOG(LogWebMMoviePlayer, Verbose, TEXT("No Movie to start."));
+		UE_LOG(LogWebMMoviePlayer, Log, TEXT("No Movie to start."));
 		return false;
 	}
 }
@@ -133,7 +162,9 @@ bool FWebMMovieStreamer::IsLastMovieInPlaylist()
 
 bool FWebMMovieStreamer::Tick(float InDeltaTime)
 {
-	if (bPlaying)
+	AudioBackend->Tick(InDeltaTime);
+
+	if (LIKELY(bPlaying && !bPaused))
 	{
 		if (TicksLeftToWaitPostCompletion)
 		{
@@ -161,6 +192,11 @@ bool FWebMMovieStreamer::Tick(float InDeltaTime)
 		}
 		return false;
 	}
+	else if (bPaused)
+	{
+		// continue ticking
+		return false;
+	}
 	else
 	{
 		// We're done playing
@@ -245,7 +281,7 @@ bool FWebMMovieStreamer::SendAudio(float InDeltaTime)
 	while (bFoundSample && AudioSample)
 	{
 		FWebMMediaAudioSample* WebMSample = StaticCast<FWebMMediaAudioSample*>(AudioSample.Get());
-		AudioBackend->SendAudio(WebMSample->GetDataBuffer().GetData(), WebMSample->GetDataBuffer().Num());
+		AudioBackend->SendAudio(WebMSample->GetTime(), WebMSample->GetDataBuffer().GetData(), WebMSample->GetDataBuffer().Num());
 
 		bFoundSample = Samples->FetchAudio(TimeRange, AudioSample);
 	}
@@ -293,4 +329,31 @@ void FWebMMovieStreamer::AddAudioSampleFromDecodingThread(TSharedRef<FWebMMediaA
 	Samples->AddAudio(Sample);
 }
 
+void FWebMMovieStreamer::HandleApplicationWillEnterBackground()
+{
+	UE_LOG(LogWebMMoviePlayer, Log, TEXT("FWebMMovieStreamer::HandleApplicationWillEnterBackground, pausing - we were %s"), bPaused ? TEXT("paused") : TEXT("not paused"));
+
+	if (!bPaused)
+	{
+		TimeAtEnteringBackground = FPlatformTime::Seconds();
+		AudioBackend->Pause(true);
+		bPaused = true;
+	}
+}
+
+void FWebMMovieStreamer::HandleApplicationHasEnteredForeground()
+{
+	UE_LOG(LogWebMMoviePlayer, Log, TEXT("FWebMMovieStreamer::HandleApplicationHasEnteredForeground, unpausing - we were %s"), bPaused ? TEXT("paused") : TEXT("not paused"));
+
+	if (bPaused)
+	{
+		double PrevStartTimeForLog = StartTime;
+		double TimeSpentInBackground = FPlatformTime::Seconds() - TimeAtEnteringBackground;
+		StartTime += TimeSpentInBackground;
+		AudioBackend->Pause(false);
+		bPaused = false;
+		UE_LOG(LogWebMMoviePlayer, Log, TEXT("Spent %f seconds in background, adjusting start time from %f to %f"), TimeSpentInBackground, PrevStartTimeForLog, StartTime);
+	}
+}
+
 #endif // WITH_WEBM_LIBS
diff --git a/Engine/Plugins/Runtime/WebMMoviePlayer/Source/WebMMoviePlayer/Private/WebMMovieStreamer.h b/Engine/Plugins/Runtime/WebMMoviePlayer/Source/WebMMoviePlayer/Private/WebMMovieStreamer.h
--- a/Engine/Plugins/Runtime/WebMMoviePlayer/Source/WebMMoviePlayer/Private/WebMMovieStreamer.h
+++ b/Engine/Plugins/Runtime/WebMMoviePlayer/Source/WebMMoviePlayer/Private/WebMMovieStreamer.h
@@ -51,9 +51,18 @@ private:
 	TSharedPtr<FMediaSamples, ESPMode::ThreadSafe> Samples;
 	TSharedPtr<FMovieViewport> Viewport;
 	TSharedPtr<FSlateTexture2DRHIRef, ESPMode::ThreadSafe> SlateVideoTexture;
+
+	/** Limits amount of frames being processed to save video memory used by the player */
 	int32 VideoFramesCurrentlyProcessing;
+
+	/** Movie start time */
 	double StartTime;
+	/** Timestamp when we entered background */
+	double TimeAtEnteringBackground;
+	/** Whether a movie is being played */
 	bool bPlaying;
+	/** We can be paused e.g. while in background */
+	bool bPaused;
 
 	/** 
 	 * Number of ticks to wait after the movie is complete before moving on to the next one. 
@@ -67,6 +76,19 @@ private:
 	bool DisplayFrames(float InDeltaTime);
 	bool SendAudio(float InDeltaTime);
 	bool ReadMoreFrames();
+
+	/** Callback for when the application resumed in the foreground. */
+	void HandleApplicationHasEnteredForeground();
+
+	/** Callback for when the application is being paused in the background. */
+	void HandleApplicationWillEnterBackground();
+
+	/** Foreground/background delegate for pause. */
+	FDelegateHandle PauseHandle;
+
+	/** Foreground/background delegate for resume. */
+	FDelegateHandle ResumeHandle;
+
 };
 
 #endif // WITH_WEBM_LIBS
