Fix CIS (WebMAudioBackendNull).

diff --git a/Engine/Plugins/Runtime/WebMMoviePlayer/Source/WebMMoviePlayer/Private/Audio/Null/WebMAudioBackend.h b/Engine/Plugins/Runtime/WebMMoviePlayer/Source/WebMMoviePlayer/Private/Audio/Null/WebMAudioBackend.h
--- a/Engine/Plugins/Runtime/WebMMoviePlayer/Source/WebMMoviePlayer/Private/Audio/Null/WebMAudioBackend.h
+++ b/Engine/Plugins/Runtime/WebMMoviePlayer/Source/WebMMoviePlayer/Private/Audio/Null/WebMAudioBackend.h
@@ -7,12 +7,20 @@
 class FWebMAudioBackendNull
 {
 public:
+
+	enum class EStreamState
+	{
+		NewMovie,
+		ResumeAfterPause
+	};
+
 	bool InitializePlatform() { return true;  }
 	void ShutdownPlatform() {}
-	bool StartStreaming(int32 SampleRate, int32 NumOfChannels) { return true; }
+	bool StartStreaming(int32 SampleRate, int32 NumOfChannels, EStreamState StreamState) { return true; }
 	void StopStreaming() {}
-	bool SendAudio(const uint8* Buffer, size_t BufferSize) { return true; }
+	bool SendAudio(const FTimespan& Timespan, const uint8* Buffer, size_t BufferSize) { return true; }
 
+	void Pause(bool bPause) {}
 	bool IsPaused() const { return false; }
 	void Tick(float DeltaTime) {}
 };
diff --git a/Engine/Plugins/Runtime/WebMMoviePlayer/Source/WebMMoviePlayer/Private/WebMMovieStreamer.cpp b/Engine/Plugins/Runtime/WebMMoviePlayer/Source/WebMMoviePlayer/Private/WebMMovieStreamer.cpp
--- a/Engine/Plugins/Runtime/WebMMoviePlayer/Source/WebMMoviePlayer/Private/WebMMovieStreamer.cpp
+++ b/Engine/Plugins/Runtime/WebMMoviePlayer/Source/WebMMoviePlayer/Private/WebMMovieStreamer.cpp
@@ -136,7 +136,7 @@ bool FWebMMovieStreamer::StartNextMovie()
 		AudioDecoder->Initialize(DefaultAudioTrack.CodecName, DefaultAudioTrack.SampleRate, DefaultAudioTrack.NumOfChannels, DefaultAudioTrack.CodecPrivateData, DefaultAudioTrack.CodecPrivateDataSize);
 		VideoDecoder->Initialize(DefaultVideoTrack.CodecName);
 
-		AudioBackend->StartStreaming(DefaultAudioTrack.SampleRate, DefaultAudioTrack.NumOfChannels, FWebMAudioBackendSDL::EStreamState::NewMovie);
+		AudioBackend->StartStreaming(DefaultAudioTrack.SampleRate, DefaultAudioTrack.NumOfChannels, FWebMAudioBackend::EStreamState::NewMovie);
 
 		StartTime = FPlatformTime::Seconds();
 		bPlaying = true;
