Fix CIS (non-Unity, CoreDelegates in WebMMovieStreamer) (UE-79130).

diff --git a/Engine/Plugins/Runtime/WebMMoviePlayer/Source/WebMMoviePlayer/Private/WebMMovieStreamer.cpp b/Engine/Plugins/Runtime/WebMMoviePlayer/Source/WebMMoviePlayer/Private/WebMMovieStreamer.cpp
--- a/Engine/Plugins/Runtime/WebMMoviePlayer/Source/WebMMoviePlayer/Private/WebMMovieStreamer.cpp
+++ b/Engine/Plugins/Runtime/WebMMoviePlayer/Source/WebMMoviePlayer/Private/WebMMovieStreamer.cpp
@@ -16,6 +16,7 @@
 #include "WebMContainer.h"
 #include "WebMMediaAudioSample.h"
 #include "WebMMediaTextureSample.h"
+#include "Misc/CoreDelegates.h"
 #endif // WITH_WEBM_LIBS
 
 DEFINE_LOG_CATEGORY(LogWebMMoviePlayer);
