Fixing WebMMoviePlayer to support Platform Extensions (UE-82685).

diff --git a/Engine/Plugins/Runtime/WebMMoviePlayer/Source/WebMMoviePlayer/Private/Audio/Linux/LinuxWebMAudioBackendDefines.h b/Engine/Plugins/Runtime/WebMMoviePlayer/Source/WebMMoviePlayer/Private/Audio/Linux/LinuxWebMAudioBackendDefines.h
new file mode 100644
--- /dev/null
+++ b/Engine/Plugins/Runtime/WebMMoviePlayer/Source/WebMMoviePlayer/Private/Audio/Linux/LinuxWebMAudioBackendDefines.h
@@ -0,0 +1,7 @@
+// Copyright 1998-2019 Epic Games, Inc. All Rights Reserved.
+
+#pragma once
+
+#include "WebMAudioBackend.h"
+
+using FWebMAudioBackend = FWebMAudioBackendSDL;
diff --git a/Engine/Plugins/Runtime/WebMMoviePlayer/Source/WebMMoviePlayer/Private/Audio/Linux/WebMAudioBackendDefines.h b/Engine/Plugins/Runtime/WebMMoviePlayer/Source/WebMMoviePlayer/Private/Audio/Linux/WebMAudioBackendDefines.h
deleted file mode 100644
diff --git a/Engine/Plugins/Runtime/WebMMoviePlayer/Source/WebMMoviePlayer/Private/Audio/Mac/MacWebMAudioBackendDefines.h b/Engine/Plugins/Runtime/WebMMoviePlayer/Source/WebMMoviePlayer/Private/Audio/Mac/MacWebMAudioBackendDefines.h
new file mode 100644
--- /dev/null
+++ b/Engine/Plugins/Runtime/WebMMoviePlayer/Source/WebMMoviePlayer/Private/Audio/Mac/MacWebMAudioBackendDefines.h
@@ -0,0 +1,5 @@
+// Copyright 1998-2019 Epic Games, Inc. All Rights Reserved.
+
+#pragma once
+
+#include "../Null/WebMAudioBackendDefines.h"
diff --git a/Engine/Plugins/Runtime/WebMMoviePlayer/Source/WebMMoviePlayer/Private/Audio/Windows/WindowsWebMAudioBackendDefines.h b/Engine/Plugins/Runtime/WebMMoviePlayer/Source/WebMMoviePlayer/Private/Audio/Windows/WindowsWebMAudioBackendDefines.h
new file mode 100644
--- /dev/null
+++ b/Engine/Plugins/Runtime/WebMMoviePlayer/Source/WebMMoviePlayer/Private/Audio/Windows/WindowsWebMAudioBackendDefines.h
@@ -0,0 +1,5 @@
+// Copyright 1998-2019 Epic Games, Inc. All Rights Reserved.
+
+#pragma once
+
+#include "../Null/WebMAudioBackendDefines.h"
diff --git a/Engine/Plugins/Runtime/WebMMoviePlayer/Source/WebMMoviePlayer/Private/WebMMovieStreamer.h b/Engine/Plugins/Runtime/WebMMoviePlayer/Source/WebMMoviePlayer/Private/WebMMovieStreamer.h
--- a/Engine/Plugins/Runtime/WebMMoviePlayer/Source/WebMMoviePlayer/Private/WebMMovieStreamer.h
+++ b/Engine/Plugins/Runtime/WebMMoviePlayer/Source/WebMMoviePlayer/Private/WebMMovieStreamer.h
@@ -6,7 +6,7 @@
 
 #include "MoviePlayer.h"
 #include "WebMAudioBackend.h"
-#include "WebMAudioBackendDefines.h"
+#include COMPILED_PLATFORM_HEADER(WebMAudioBackendDefines.h)
 #include "WebMMediaFrame.h"
 #include "Containers/Queue.h"
 #include "WebMSamplesSink.h"
diff --git a/Engine/Plugins/Runtime/WebMMoviePlayer/Source/WebMMoviePlayer/WebMMoviePlayer.Build.cs b/Engine/Plugins/Runtime/WebMMoviePlayer/Source/WebMMoviePlayer/WebMMoviePlayer.Build.cs
--- a/Engine/Plugins/Runtime/WebMMoviePlayer/Source/WebMMoviePlayer/WebMMoviePlayer.Build.cs
+++ b/Engine/Plugins/Runtime/WebMMoviePlayer/Source/WebMMoviePlayer/WebMMoviePlayer.Build.cs
@@ -1,45 +1,43 @@
 // Copyright 1998-2019 Epic Games, Inc. All Rights Reserved.
+using UnrealBuildTool;
 using System;
 using System.IO;
 
-namespace UnrealBuildTool.Rules
+public class WebMMoviePlayer : ModuleRules
 {
-	public class WebMMoviePlayer : ModuleRules
+	public WebMMoviePlayer(ReadOnlyTargetRules Target) : base(Target)
 	{
-		public WebMMoviePlayer(ReadOnlyTargetRules Target) : base(Target)
-		{
-			PrivateDependencyModuleNames.AddRange(
-				new string[]
-				{
-					"Core",
-					"Engine",
-					"MoviePlayer",
-					"RenderCore",
-					"RHI",
-					"SlateCore",
-					"Slate",
-					"MediaUtils",
-					"WebMMedia",
-				});
-
-			if (Target.IsInPlatformGroup(UnrealPlatformGroup.Unix))
+		PrivateDependencyModuleNames.AddRange(
+			new string[]
 			{
-				string PlatformName = Target.IsInPlatformGroup(UnrealPlatformGroup.Linux) ?
-					"Linux" : Target.Platform.ToString();
+				"Core",
+				"Engine",
+				"MoviePlayer",
+				"RenderCore",
+				"RHI",
+				"SlateCore",
+				"Slate",
+				"MediaUtils",
+				"WebMMedia",
+			});
 
-				PrivateDependencyModuleNames.Add("SDL2");
+		if (Target.IsInPlatformGroup(UnrealPlatformGroup.Unix))
+		{
+			string PlatformName = Target.IsInPlatformGroup(UnrealPlatformGroup.Linux) ?
+				"Linux" : Target.Platform.ToString();
 
-				PrivateIncludePaths.Add("WebMMoviePlayer/Private/Audio/Unix");
-				PrivateIncludePaths.Add("WebMMoviePlayer/Private/Audio/" + PlatformName);
-			}
-			else
-			{
-				PrivateIncludePaths.Add("WebMMoviePlayer/Private/Audio/Null");
-			}
+			PrivateDependencyModuleNames.Add("SDL2");
 
-			// Some Linux architectures don't have the libs built yet
-			bool bHaveWebMlibs = (!Target.IsInPlatformGroup(UnrealPlatformGroup.Unix) || Target.Architecture.StartsWith("x86_64"));
-			PublicDefinitions.Add("WITH_WEBM_LIBS=" + (bHaveWebMlibs ? "1" : "0"));
+			PrivateIncludePaths.Add("WebMMoviePlayer/Private/Audio/Unix");
 		}
+		else
+		{
+			PrivateIncludePaths.Add("WebMMoviePlayer/Private/Audio/Null");
+		}
+		PrivateIncludePaths.Add("WebMMoviePlayer/Private/Audio");
+
+		// Some Linux architectures don't have the libs built yet
+		bool bHaveWebMlibs = (!Target.IsInPlatformGroup(UnrealPlatformGroup.Unix) || Target.Architecture.StartsWith("x86_64"));
+		PublicDefinitions.Add("WITH_WEBM_LIBS=" + (bHaveWebMlibs ? "1" : "0"));
 	}
 }
