Fix warning introduced by CL# 10756027
#rnx

diff --git a/Engine/Plugins/Runtime/WebMMoviePlayer/Source/WebMMoviePlayer/WebMMoviePlayer.Build.cs b/Engine/Plugins/Runtime/WebMMoviePlayer/Source/WebMMoviePlayer/WebMMoviePlayer.Build.cs
--- a/Engine/Plugins/Runtime/WebMMoviePlayer/Source/WebMMoviePlayer/WebMMoviePlayer.Build.cs
+++ b/Engine/Plugins/Runtime/WebMMoviePlayer/Source/WebMMoviePlayer/WebMMoviePlayer.Build.cs
@@ -23,9 +23,6 @@ public class WebMMoviePlayer : ModuleRules
 
 		if (Target.IsInPlatformGroup(UnrealPlatformGroup.Unix))
 		{
-			string PlatformName = Target.IsInPlatformGroup(UnrealPlatformGroup.Linux) ?
-				"Linux" : Target.Platform.ToString();
-
 			PrivateDependencyModuleNames.Add("SDL2");
 
 			PrivateIncludePaths.Add("WebMMoviePlayer/Private/Audio/Unix");
