fix WITH_WEBM_LIBS mismatched define on Windows Arm64

#rnx

diff --git a/Engine/Plugins/Runtime/WebMMoviePlayer/Source/WebMMoviePlayer/WebMMoviePlayer.Build.cs b/Engine/Plugins/Runtime/WebMMoviePlayer/Source/WebMMoviePlayer/WebMMoviePlayer.Build.cs
--- a/Engine/Plugins/Runtime/WebMMoviePlayer/Source/WebMMoviePlayer/WebMMoviePlayer.Build.cs
+++ b/Engine/Plugins/Runtime/WebMMoviePlayer/Source/WebMMoviePlayer/WebMMoviePlayer.Build.cs
@@ -35,7 +35,7 @@ public class WebMMoviePlayer : ModuleRules
 
 		bool bHaveWebMlibs = 
 			Target.Platform == UnrealTargetPlatform.Mac ||
-			(Target.Platform == UnrealTargetPlatform.Win64 && Target.WindowsPlatform.Architecture != UnrealArch.Arm64) ||
+			Target.Platform == UnrealTargetPlatform.Win64 ||
 			(Target.IsInPlatformGroup(UnrealPlatformGroup.Unix) && Target.Architecture == UnrealArch.X64);
 		PublicDefinitions.Add("WITH_WEBM_LIBS=" + (bHaveWebMlibs ? "1" : "0"));
 	}
